Merging using Dev-VR_to_Release-4.21

diff --git a/Engine/Plugins/Experimental/AppleImageUtils/Source/AppleImageUtils/Private/AppleImageUtilsBlueprintProxy.cpp b/Engine/Plugins/Experimental/AppleImageUtils/Source/AppleImageUtils/Private/AppleImageUtilsBlueprintProxy.cpp
--- a/Engine/Plugins/Experimental/AppleImageUtils/Source/AppleImageUtils/Private/AppleImageUtilsBlueprintProxy.cpp
+++ b/Engine/Plugins/Experimental/AppleImageUtils/Source/AppleImageUtils/Private/AppleImageUtilsBlueprintProxy.cpp
@@ -37,42 +37,42 @@ void UAppleImageUtilsBaseAsyncTaskBlueprintProxy::Tick(float DeltaTime)
 	}
 }
 
-UAppleImageUtilsBaseAsyncTaskBlueprintProxy* UAppleImageUtilsBaseAsyncTaskBlueprintProxy::CreateProxyObjectForConvertToJPEG(UTexture* SourceImage, int32 Quality, bool bWantColor, bool bUseGpu)
+UAppleImageUtilsBaseAsyncTaskBlueprintProxy* UAppleImageUtilsBaseAsyncTaskBlueprintProxy::CreateProxyObjectForConvertToJPEG(UTexture* SourceImage, int32 Quality, bool bWantColor, bool bUseGpu, float Scale)
 {
 	UAppleImageUtilsBaseAsyncTaskBlueprintProxy* Proxy = NewObject<UAppleImageUtilsBaseAsyncTaskBlueprintProxy>();
 	Proxy->SetFlags(RF_StrongRefOnFrame);
 
-	Proxy->ConversionTask = IAppleImageUtilsPlugin::Get().ConvertToJPEG(SourceImage, Quality, bWantColor, bUseGpu);
+	Proxy->ConversionTask = IAppleImageUtilsPlugin::Get().ConvertToJPEG(SourceImage, Quality, bWantColor, bUseGpu, Scale);
 
 	return Proxy;
 }
 
-UAppleImageUtilsBaseAsyncTaskBlueprintProxy* UAppleImageUtilsBaseAsyncTaskBlueprintProxy::CreateProxyObjectForConvertToHEIF(UTexture* SourceImage, int32 Quality, bool bWantColor, bool bUseGpu)
+UAppleImageUtilsBaseAsyncTaskBlueprintProxy* UAppleImageUtilsBaseAsyncTaskBlueprintProxy::CreateProxyObjectForConvertToHEIF(UTexture* SourceImage, int32 Quality, bool bWantColor, bool bUseGpu, float Scale)
 {
 	UAppleImageUtilsBaseAsyncTaskBlueprintProxy* Proxy = NewObject<UAppleImageUtilsBaseAsyncTaskBlueprintProxy>();
 	Proxy->SetFlags(RF_StrongRefOnFrame);
 
-	Proxy->ConversionTask = IAppleImageUtilsPlugin::Get().ConvertToHEIF(SourceImage, Quality, bWantColor, bUseGpu);
+	Proxy->ConversionTask = IAppleImageUtilsPlugin::Get().ConvertToHEIF(SourceImage, Quality, bWantColor, bUseGpu, Scale);
 
 	return Proxy;
 }
 
-UAppleImageUtilsBaseAsyncTaskBlueprintProxy* UAppleImageUtilsBaseAsyncTaskBlueprintProxy::CreateProxyObjectForConvertToTIFF(UTexture* SourceImage, bool bWantColor, bool bUseGpu)
+UAppleImageUtilsBaseAsyncTaskBlueprintProxy* UAppleImageUtilsBaseAsyncTaskBlueprintProxy::CreateProxyObjectForConvertToTIFF(UTexture* SourceImage, bool bWantColor, bool bUseGpu, float Scale)
 {
 	UAppleImageUtilsBaseAsyncTaskBlueprintProxy* Proxy = NewObject<UAppleImageUtilsBaseAsyncTaskBlueprintProxy>();
 	Proxy->SetFlags(RF_StrongRefOnFrame);
 
-	Proxy->ConversionTask = IAppleImageUtilsPlugin::Get().ConvertToTIFF(SourceImage, bWantColor, bUseGpu);
+	Proxy->ConversionTask = IAppleImageUtilsPlugin::Get().ConvertToTIFF(SourceImage, bWantColor, bUseGpu, Scale);
 
 	return Proxy;
 }
 
-UAppleImageUtilsBaseAsyncTaskBlueprintProxy* UAppleImageUtilsBaseAsyncTaskBlueprintProxy::CreateProxyObjectForConvertToPNG(UTexture* SourceImage, bool bWantColor, bool bUseGpu)
+UAppleImageUtilsBaseAsyncTaskBlueprintProxy* UAppleImageUtilsBaseAsyncTaskBlueprintProxy::CreateProxyObjectForConvertToPNG(UTexture* SourceImage, bool bWantColor, bool bUseGpu, float Scale)
 {
 	UAppleImageUtilsBaseAsyncTaskBlueprintProxy* Proxy = NewObject<UAppleImageUtilsBaseAsyncTaskBlueprintProxy>();
 	Proxy->SetFlags(RF_StrongRefOnFrame);
 
-	Proxy->ConversionTask = IAppleImageUtilsPlugin::Get().ConvertToPNG(SourceImage, bWantColor, bUseGpu);
+	Proxy->ConversionTask = IAppleImageUtilsPlugin::Get().ConvertToPNG(SourceImage, bWantColor, bUseGpu, Scale);
 
 	return Proxy;
 }
diff --git a/Engine/Plugins/Experimental/AppleImageUtils/Source/AppleImageUtils/Private/AppleImageUtilsPlugin.cpp b/Engine/Plugins/Experimental/AppleImageUtils/Source/AppleImageUtils/Private/AppleImageUtilsPlugin.cpp
--- a/Engine/Plugins/Experimental/AppleImageUtils/Source/AppleImageUtils/Private/AppleImageUtilsPlugin.cpp
+++ b/Engine/Plugins/Experimental/AppleImageUtils/Source/AppleImageUtils/Private/AppleImageUtilsPlugin.cpp
@@ -26,10 +26,10 @@ class FAppleImageUtilsPlugin :
 	virtual void StartupModule() override;
 	virtual void ShutdownModule() override;
 
-	virtual TSharedPtr<FAppleImageUtilsConversionTaskBase, ESPMode::ThreadSafe> ConvertToJPEG(UTexture* SourceImage, int32 Quality = 85, bool bWantColor = true, bool bUseGpu = true) override;
-	virtual TSharedPtr<FAppleImageUtilsConversionTaskBase, ESPMode::ThreadSafe> ConvertToHEIF(UTexture* SourceImage, int32 Quality = 85, bool bWantColor = true, bool bUseGpu = true) override;
-	virtual TSharedPtr<FAppleImageUtilsConversionTaskBase, ESPMode::ThreadSafe> ConvertToPNG(UTexture* SourceImage, bool bWantColor = true, bool bUseGpu = true) override;
-	virtual TSharedPtr<FAppleImageUtilsConversionTaskBase, ESPMode::ThreadSafe> ConvertToTIFF(UTexture* SourceImage, bool bWantColor = true, bool bUseGpu = true) override;
+	virtual TSharedPtr<FAppleImageUtilsConversionTaskBase, ESPMode::ThreadSafe> ConvertToJPEG(UTexture* SourceImage, int32 Quality = 85, bool bWantColor = true, bool bUseGpu = true, float Scale = 1.f) override;
+	virtual TSharedPtr<FAppleImageUtilsConversionTaskBase, ESPMode::ThreadSafe> ConvertToHEIF(UTexture* SourceImage, int32 Quality = 85, bool bWantColor = true, bool bUseGpu = true, float Scale = 1.f) override;
+	virtual TSharedPtr<FAppleImageUtilsConversionTaskBase, ESPMode::ThreadSafe> ConvertToPNG(UTexture* SourceImage, bool bWantColor = true, bool bUseGpu = true, float Scale = 1.f) override;
+	virtual TSharedPtr<FAppleImageUtilsConversionTaskBase, ESPMode::ThreadSafe> ConvertToTIFF(UTexture* SourceImage, bool bWantColor = true, bool bUseGpu = true, float Scale = 1.f) override;
 #if SUPPORTS_IMAGE_UTILS_1_0
 	virtual CGImageRef UTexture2DToCGImage(UTexture2D* Source) override;
 #endif
@@ -115,15 +115,17 @@ static inline CGColorSpaceRef ToColorSpace(bool bWantColor)
 	return bWantColor ? CGColorSpaceCreateWithName(kCGColorSpaceSRGB) : CGColorSpaceCreateWithName(kCGColorSpaceGenericGrayGamma2_2);
 }
 
-static inline CIImage* ToImage(IAppleImageInterface* AppleImageInterface)
+static inline CIImage* ToImage(IAppleImageInterface* AppleImageInterface, float Scale)
 {
 	check(AppleImageInterface != nullptr);
 
+	CIImage* Image = nullptr;
 	switch (AppleImageInterface->GetTextureType())
 	{
 		case EAppleTextureType::Image:
 		{
-			return AppleImageInterface->GetImage();
+			Image = AppleImageInterface->GetImage();
+			break;
 		}
 
 		case EAppleTextureType::PixelBuffer:
@@ -131,7 +133,7 @@ static inline CIImage* ToImage(IAppleImageInterface* AppleImageInterface)
 			CVPixelBufferRef PixelBuffer = AppleImageInterface->GetPixelBuffer();
 			if (PixelBuffer != nullptr)
 			{
-				return [[CIImage imageWithCVPixelBuffer: PixelBuffer] autorelease];
+				Image = [[CIImage imageWithCVPixelBuffer: PixelBuffer] autorelease];
 			}
 			break;
 		}
@@ -141,12 +143,23 @@ static inline CIImage* ToImage(IAppleImageInterface* AppleImageInterface)
 			IOSurfaceRef Surface = AppleImageInterface->GetSurface();
 			if (Surface != nullptr)
 			{
-				return [[CIImage imageWithIOSurface: Surface] autorelease];
+				Image = [[CIImage imageWithIOSurface: Surface] autorelease];
 			}
 			break;
 		}
 	}
-	return nullptr;
+	// Handle scaling if requested
+	if (Scale != 1.f)
+	{
+		CGRect Rect = Image.extent;
+		float AspectRatio = (float)Rect.size.width / (float)Rect.size.height;
+		CIFilter* ScaleFilter = [CIFilter filterWithName: @"CILanczosScaleTransform"];
+		[ScaleFilter setValue: Image forKey: kCIInputImageKey];
+		[ScaleFilter setValue: @(Scale) forKey: kCIInputScaleKey];
+		[ScaleFilter setValue: @(AspectRatio) forKey: kCIInputAspectRatioKey];
+		Image = ScaleFilter.outputImage;
+	}
+	return Image;
 }
 #endif
 
@@ -164,7 +177,7 @@ static inline bool CanBeConverted(IAppleImageInterface* AppleImageInterface)
 	return false;
 }
 
-TSharedPtr<FAppleImageUtilsConversionTaskBase, ESPMode::ThreadSafe> FAppleImageUtilsPlugin::ConvertToJPEG(UTexture* SourceImage, int32 Quality, bool bWantColor, bool bUseGpu)
+TSharedPtr<FAppleImageUtilsConversionTaskBase, ESPMode::ThreadSafe> FAppleImageUtilsPlugin::ConvertToJPEG(UTexture* SourceImage, int32 Quality, bool bWantColor, bool bUseGpu, float Scale)
 {
 	// Make sure our interface is supported
 	IAppleImageInterface* AppleImage = Cast<IAppleImageInterface>(SourceImage);
@@ -184,14 +197,14 @@ TSharedPtr<FAppleImageUtilsConversionTaskBase, ESPMode::ThreadSafe> FAppleImageU
 		return MakeShared<FAppleImageUtilsFailedConversionTask, ESPMode::ThreadSafe>(TEXT("ConvertToJPEG requires iOS 10.0+ or macOS 10.12+"));
 	}
 
-	AsyncTask(ENamedThreads::AnyBackgroundThreadNormalTask, [ConversionTask, Quality, bWantColor, bUseGpu]()
+	AsyncTask(ENamedThreads::AnyBackgroundThreadNormalTask, [ConversionTask, Quality, bWantColor, bUseGpu, Scale]()
 	{
 		IAppleImageInterface* AppleImageInterface = Cast<IAppleImageInterface>(ConversionTask->SourceImage);
 
 		// Convert to Apple objects
 		CIContext* ConversionContext = [[[CIContext alloc] initWithOptions: ToCpuDictionary(bUseGpu)] autorelease];
 		CGColorSpaceRef ColorSpace = ToColorSpace(bWantColor);
-		CIImage* Image = ToImage(AppleImageInterface);
+		CIImage* Image = ToImage(AppleImageInterface, Scale);
 
 		// This will perform the work on the GPU or inline on this thread
 		NSData* ConvertedData = [ConversionContext JPEGRepresentationOfImage: Image colorSpace: ColorSpace options: ToQualityDictionary(Quality)];
@@ -215,7 +228,7 @@ TSharedPtr<FAppleImageUtilsConversionTaskBase, ESPMode::ThreadSafe> FAppleImageU
 	return ConversionTask;
 }
 
-TSharedPtr<FAppleImageUtilsConversionTaskBase, ESPMode::ThreadSafe> FAppleImageUtilsPlugin::ConvertToHEIF(UTexture* SourceImage, int32 Quality, bool bWantColor, bool bUseGpu)
+TSharedPtr<FAppleImageUtilsConversionTaskBase, ESPMode::ThreadSafe> FAppleImageUtilsPlugin::ConvertToHEIF(UTexture* SourceImage, int32 Quality, bool bWantColor, bool bUseGpu, float Scale)
 {
 	// Make sure our interface is supported
 	IAppleImageInterface* AppleImage = Cast<IAppleImageInterface>(SourceImage);
@@ -235,14 +248,14 @@ TSharedPtr<FAppleImageUtilsConversionTaskBase, ESPMode::ThreadSafe> FAppleImageU
 		return MakeShared<FAppleImageUtilsFailedConversionTask, ESPMode::ThreadSafe>(TEXT("ConvertToHEIF requires iOS 11.0+ or macOS 10.13.4+"));
 	}
 
-	AsyncTask(ENamedThreads::AnyBackgroundThreadNormalTask, [ConversionTask, Quality, bWantColor, bUseGpu]()
+	AsyncTask(ENamedThreads::AnyBackgroundThreadNormalTask, [ConversionTask, Quality, bWantColor, bUseGpu, Scale]()
 	{
 		IAppleImageInterface* AppleImageInterface = Cast<IAppleImageInterface>(ConversionTask->SourceImage);
 
 		// Convert to Apple objects
 		CIContext* ConversionContext = [[[CIContext alloc] initWithOptions: ToCpuDictionary(bUseGpu)] autorelease];
 		CGColorSpaceRef ColorSpace = ToColorSpace(bWantColor);
-		CIImage* Image = ToImage(AppleImageInterface);
+		CIImage* Image = ToImage(AppleImageInterface, Scale);
 
 		// This will perform the work on the GPU or inline on this thread
 		NSData* ConvertedData = [ConversionContext HEIFRepresentationOfImage: Image format: kCIFormatARGB8 colorSpace: ColorSpace options: ToQualityDictionary(Quality)];
@@ -266,7 +279,7 @@ TSharedPtr<FAppleImageUtilsConversionTaskBase, ESPMode::ThreadSafe> FAppleImageU
 	return ConversionTask;
 }
 
-TSharedPtr<FAppleImageUtilsConversionTaskBase, ESPMode::ThreadSafe> FAppleImageUtilsPlugin::ConvertToPNG(UTexture* SourceImage, bool bWantColor, bool bUseGpu)
+TSharedPtr<FAppleImageUtilsConversionTaskBase, ESPMode::ThreadSafe> FAppleImageUtilsPlugin::ConvertToPNG(UTexture* SourceImage, bool bWantColor, bool bUseGpu, float Scale)
 {
 	// Make sure our interface is supported
 	IAppleImageInterface* AppleImage = Cast<IAppleImageInterface>(SourceImage);
@@ -286,14 +299,14 @@ TSharedPtr<FAppleImageUtilsConversionTaskBase, ESPMode::ThreadSafe> FAppleImageU
 		return MakeShared<FAppleImageUtilsFailedConversionTask, ESPMode::ThreadSafe>(TEXT("ConvertToPNG requires iOS 11.0+ or macOS 10.13+"));
 	}
 
-	AsyncTask(ENamedThreads::AnyBackgroundThreadNormalTask, [ConversionTask, bWantColor, bUseGpu]()
+	AsyncTask(ENamedThreads::AnyBackgroundThreadNormalTask, [ConversionTask, bWantColor, bUseGpu, Scale]()
 	{
 		IAppleImageInterface* AppleImageInterface = Cast<IAppleImageInterface>(ConversionTask->SourceImage);
 
 		// Convert to Apple objects
 		CIContext* ConversionContext = [[[CIContext alloc] initWithOptions: ToCpuDictionary(bUseGpu)] autorelease];
 		CGColorSpaceRef ColorSpace = ToColorSpace(bWantColor);
-		CIImage* Image = ToImage(AppleImageInterface);
+		CIImage* Image = ToImage(AppleImageInterface, Scale);
 
 		// This will perform the work on the GPU or inline on this thread
 		NSData* ConvertedData = [ConversionContext PNGRepresentationOfImage: Image format: kCIFormatARGB8 colorSpace: ColorSpace options: @{}];
@@ -317,7 +330,7 @@ TSharedPtr<FAppleImageUtilsConversionTaskBase, ESPMode::ThreadSafe> FAppleImageU
 	return ConversionTask;
 }
 
-TSharedPtr<FAppleImageUtilsConversionTaskBase, ESPMode::ThreadSafe> FAppleImageUtilsPlugin::ConvertToTIFF(UTexture* SourceImage, bool bWantColor, bool bUseGpu)
+TSharedPtr<FAppleImageUtilsConversionTaskBase, ESPMode::ThreadSafe> FAppleImageUtilsPlugin::ConvertToTIFF(UTexture* SourceImage, bool bWantColor, bool bUseGpu, float Scale)
 {
 	// Make sure our interface is supported
 	IAppleImageInterface* AppleImage = Cast<IAppleImageInterface>(SourceImage);
@@ -337,14 +350,14 @@ TSharedPtr<FAppleImageUtilsConversionTaskBase, ESPMode::ThreadSafe> FAppleImageU
 		return MakeShared<FAppleImageUtilsFailedConversionTask, ESPMode::ThreadSafe>(TEXT("ConvertToTIFF requires iOS 10.0+ or macOS 10.12+"));
 	}
 
-	AsyncTask(ENamedThreads::AnyBackgroundThreadNormalTask, [ConversionTask, bWantColor, bUseGpu]()
+	AsyncTask(ENamedThreads::AnyBackgroundThreadNormalTask, [ConversionTask, bWantColor, bUseGpu, Scale]()
 	{
 		IAppleImageInterface* AppleImageInterface = Cast<IAppleImageInterface>(ConversionTask->SourceImage);
 
 		// Convert to Apple objects
 		CIContext* ConversionContext = [[[CIContext alloc] initWithOptions: ToCpuDictionary(bUseGpu)] autorelease];
 		CGColorSpaceRef ColorSpace = ToColorSpace(bWantColor);
-		CIImage* Image = ToImage(AppleImageInterface);
+		CIImage* Image = ToImage(AppleImageInterface, Scale);
 
 		// This will perform the work on the GPU or inline on this thread
 		NSData* ConvertedData = [ConversionContext TIFFRepresentationOfImage: Image format: kCIFormatARGB8 colorSpace: ColorSpace options: @{}];
diff --git a/Engine/Plugins/Experimental/AppleImageUtils/Source/AppleImageUtils/Public/AppleImageUtilsBlueprintProxy.h b/Engine/Plugins/Experimental/AppleImageUtils/Source/AppleImageUtils/Public/AppleImageUtilsBlueprintProxy.h
--- a/Engine/Plugins/Experimental/AppleImageUtils/Source/AppleImageUtils/Public/AppleImageUtilsBlueprintProxy.h
+++ b/Engine/Plugins/Experimental/AppleImageUtils/Source/AppleImageUtils/Public/AppleImageUtilsBlueprintProxy.h
@@ -49,9 +49,10 @@ public:
 	 * @param Quality the quality level to compress to
 	 * @param bWantColor whether the JPEG is color (true) or monochrome (false)
 	 * @param bUseGpu whether to use the GPU (true) or the CPU (false) to compress
+	 * @param Scale whether to scale the image before conversion, defaults to no scaling
 	 */
 	UFUNCTION(BlueprintCallable, Meta=(BlueprintInternalUseOnly="true", DisplayName="Convert To JPEG"), Category="Image Conversion")
-	static UAppleImageUtilsBaseAsyncTaskBlueprintProxy* CreateProxyObjectForConvertToJPEG(UTexture* SourceImage, int32 Quality = 85, bool bWantColor = true, bool bUseGpu = true);
+	static UAppleImageUtilsBaseAsyncTaskBlueprintProxy* CreateProxyObjectForConvertToJPEG(UTexture* SourceImage, int32 Quality = 85, bool bWantColor = true, bool bUseGpu = true, float Scale = 1.f);
 
 	/**
 	 * Converts a image to an array of HEIF data in a background task
@@ -60,9 +61,10 @@ public:
 	 * @param Quality the quality level to compress to
 	 * @param bWantColor whether the HEIF is color (true) or monochrome (false)
 	 * @param bUseGpu whether to use the GPU (true) or the CPU (false) to compress
+	 * @param Scale whether to scale the image before conversion, defaults to no scaling
 	 */
 	UFUNCTION(BlueprintCallable, Meta=(BlueprintInternalUseOnly="true", DisplayName="Convert To HEIF"), Category="Image Conversion")
-	static UAppleImageUtilsBaseAsyncTaskBlueprintProxy* CreateProxyObjectForConvertToHEIF(UTexture* SourceImage, int32 Quality = 85, bool bWantColor = true, bool bUseGpu = true);
+	static UAppleImageUtilsBaseAsyncTaskBlueprintProxy* CreateProxyObjectForConvertToHEIF(UTexture* SourceImage, int32 Quality = 85, bool bWantColor = true, bool bUseGpu = true, float Scale = 1.f);
 
 	/**
 	 * Converts a image to an array of TIFF data in a background task
@@ -71,9 +73,10 @@ public:
 	 * @param Quality the quality level to compress to
 	 * @param bWantColor whether the TIFF is color (true) or monochrome (false)
 	 * @param bUseGpu whether to use the GPU (true) or the CPU (false) to compress
+	 * @param Scale whether to scale the image before conversion, defaults to no scaling
 	 */
 	UFUNCTION(BlueprintCallable, Meta=(BlueprintInternalUseOnly="true", DisplayName="Convert To TIFF"), Category="Image Conversion")
-	static UAppleImageUtilsBaseAsyncTaskBlueprintProxy* CreateProxyObjectForConvertToTIFF(UTexture* SourceImage, bool bWantColor = true, bool bUseGpu = true);
+	static UAppleImageUtilsBaseAsyncTaskBlueprintProxy* CreateProxyObjectForConvertToTIFF(UTexture* SourceImage, bool bWantColor = true, bool bUseGpu = true, float Scale = 1.f);
 
 	/**
 	 * Converts a image to an array of TIFF data in a background task
@@ -82,9 +85,10 @@ public:
 	 * @param Quality the quality level to compress to
 	 * @param bWantColor whether the PNG is color (true) or monochrome (false)
 	 * @param bUseGpu whether to use the GPU (true) or the CPU (false) to compress
+	 * @param Scale whether to scale the image before conversion, defaults to no scaling
 	 */
 	UFUNCTION(BlueprintCallable, Meta=(BlueprintInternalUseOnly="true", DisplayName="Convert To PNG"), Category="Image Conversion")
-	static UAppleImageUtilsBaseAsyncTaskBlueprintProxy* CreateProxyObjectForConvertToPNG(UTexture* SourceImage, bool bWantColor = true, bool bUseGpu = true);
+	static UAppleImageUtilsBaseAsyncTaskBlueprintProxy* CreateProxyObjectForConvertToPNG(UTexture* SourceImage, bool bWantColor = true, bool bUseGpu = true, float Scale = 1.f);
 
 	//~ Begin FTickableObject Interface
 	virtual void Tick(float DeltaTime) override;
diff --git a/Engine/Plugins/Experimental/AppleImageUtils/Source/AppleImageUtils/Public/IAppleImageUtilsPlugin.h b/Engine/Plugins/Experimental/AppleImageUtils/Source/AppleImageUtils/Public/IAppleImageUtilsPlugin.h
--- a/Engine/Plugins/Experimental/AppleImageUtils/Source/AppleImageUtils/Public/IAppleImageUtilsPlugin.h
+++ b/Engine/Plugins/Experimental/AppleImageUtils/Source/AppleImageUtils/Public/IAppleImageUtilsPlugin.h
@@ -85,10 +85,11 @@ public:
 	 * @param Quality the quality level to compress to
 	 * @param bWantColor whether the JPEG is color (true) or monochrome (false)
 	 * @param bUseGpu whether to use the GPU (true) or the CPU (false) to compress
+	 * @param Scale whether to scale the image before converting, defaults to no scale operation
 	 *
 	 * @return the async task that is doing the conversion
 	 */
-	virtual TSharedPtr<FAppleImageUtilsConversionTaskBase, ESPMode::ThreadSafe> ConvertToJPEG(UTexture* SourceImage, int32 Quality = 85, bool bWantColor = true, bool bUseGpu = true) = 0;
+	virtual TSharedPtr<FAppleImageUtilsConversionTaskBase, ESPMode::ThreadSafe> ConvertToJPEG(UTexture* SourceImage, int32 Quality = 85, bool bWantColor = true, bool bUseGpu = true, float Scale = 1.f) = 0;
 
 	/**
 	 * Converts a image to an array of HEIF data in a background task
@@ -100,7 +101,7 @@ public:
 	 *
 	 * @return the async task that is doing the conversion
 	 */
-	virtual TSharedPtr<FAppleImageUtilsConversionTaskBase, ESPMode::ThreadSafe> ConvertToHEIF(UTexture* SourceImage, int32 Quality = 85,  bool bWantColor = true, bool bUseGpu = true) = 0;
+	virtual TSharedPtr<FAppleImageUtilsConversionTaskBase, ESPMode::ThreadSafe> ConvertToHEIF(UTexture* SourceImage, int32 Quality = 85,  bool bWantColor = true, bool bUseGpu = true, float Scale = 1.f) = 0;
 
 	/**
 	 * Converts a image to an array of PNG data in a background task
@@ -112,7 +113,7 @@ public:
 	 *
 	 * @return the async task that is doing the conversion
 	 */
-	virtual TSharedPtr<FAppleImageUtilsConversionTaskBase, ESPMode::ThreadSafe> ConvertToPNG(UTexture* SourceImage, bool bWantColor = true, bool bUseGpu = true) = 0;
+	virtual TSharedPtr<FAppleImageUtilsConversionTaskBase, ESPMode::ThreadSafe> ConvertToPNG(UTexture* SourceImage, bool bWantColor = true, bool bUseGpu = true, float Scale = 1.f) = 0;
 
 	/**
 	 * Converts a image to an array of TIFF data in a background task
@@ -123,7 +124,7 @@ public:
 	 *
 	 * @return the async task that is doing the conversion
 	 */
-	virtual TSharedPtr<FAppleImageUtilsConversionTaskBase, ESPMode::ThreadSafe> ConvertToTIFF(UTexture* SourceImage, bool bWantColor = true, bool bUseGpu = true) = 0;
+	virtual TSharedPtr<FAppleImageUtilsConversionTaskBase, ESPMode::ThreadSafe> ConvertToTIFF(UTexture* SourceImage, bool bWantColor = true, bool bUseGpu = true, float Scale = 1.f) = 0;
 
 #if SUPPORTS_IMAGE_UTILS_1_0
 	/**
