Merging using Dev-VR_to_Release-4.21

diff --git a/Engine/Plugins/Experimental/AppleImageUtils/Source/AppleImageUtils/Private/AppleImageUtilsBlueprintProxy.cpp b/Engine/Plugins/Experimental/AppleImageUtils/Source/AppleImageUtils/Private/AppleImageUtilsBlueprintProxy.cpp
--- a/Engine/Plugins/Experimental/AppleImageUtils/Source/AppleImageUtils/Private/AppleImageUtilsBlueprintProxy.cpp
+++ b/Engine/Plugins/Experimental/AppleImageUtils/Source/AppleImageUtils/Private/AppleImageUtilsBlueprintProxy.cpp
@@ -37,42 +37,42 @@ void UAppleImageUtilsBaseAsyncTaskBlueprintProxy::Tick(float DeltaTime)
 	}
 }
 
-UAppleImageUtilsBaseAsyncTaskBlueprintProxy* UAppleImageUtilsBaseAsyncTaskBlueprintProxy::CreateProxyObjectForConvertToJPEG(UTexture* SourceImage, int32 Quality, bool bWantColor, bool bUseGpu, float Scale)
+UAppleImageUtilsBaseAsyncTaskBlueprintProxy* UAppleImageUtilsBaseAsyncTaskBlueprintProxy::CreateProxyObjectForConvertToJPEG(UTexture* SourceImage, int32 Quality, bool bWantColor, bool bUseGpu, float Scale, ETextureRotationDirection Rotate)
 {
 	UAppleImageUtilsBaseAsyncTaskBlueprintProxy* Proxy = NewObject<UAppleImageUtilsBaseAsyncTaskBlueprintProxy>();
 	Proxy->SetFlags(RF_StrongRefOnFrame);
 
-	Proxy->ConversionTask = IAppleImageUtilsPlugin::Get().ConvertToJPEG(SourceImage, Quality, bWantColor, bUseGpu, Scale);
+	Proxy->ConversionTask = IAppleImageUtilsPlugin::Get().ConvertToJPEG(SourceImage, Quality, bWantColor, bUseGpu, Scale, Rotate);
 
 	return Proxy;
 }
 
-UAppleImageUtilsBaseAsyncTaskBlueprintProxy* UAppleImageUtilsBaseAsyncTaskBlueprintProxy::CreateProxyObjectForConvertToHEIF(UTexture* SourceImage, int32 Quality, bool bWantColor, bool bUseGpu, float Scale)
+UAppleImageUtilsBaseAsyncTaskBlueprintProxy* UAppleImageUtilsBaseAsyncTaskBlueprintProxy::CreateProxyObjectForConvertToHEIF(UTexture* SourceImage, int32 Quality, bool bWantColor, bool bUseGpu, float Scale, ETextureRotationDirection Rotate)
 {
 	UAppleImageUtilsBaseAsyncTaskBlueprintProxy* Proxy = NewObject<UAppleImageUtilsBaseAsyncTaskBlueprintProxy>();
 	Proxy->SetFlags(RF_StrongRefOnFrame);
 
-	Proxy->ConversionTask = IAppleImageUtilsPlugin::Get().ConvertToHEIF(SourceImage, Quality, bWantColor, bUseGpu, Scale);
+	Proxy->ConversionTask = IAppleImageUtilsPlugin::Get().ConvertToHEIF(SourceImage, Quality, bWantColor, bUseGpu, Scale, Rotate);
 
 	return Proxy;
 }
 
-UAppleImageUtilsBaseAsyncTaskBlueprintProxy* UAppleImageUtilsBaseAsyncTaskBlueprintProxy::CreateProxyObjectForConvertToTIFF(UTexture* SourceImage, bool bWantColor, bool bUseGpu, float Scale)
+UAppleImageUtilsBaseAsyncTaskBlueprintProxy* UAppleImageUtilsBaseAsyncTaskBlueprintProxy::CreateProxyObjectForConvertToTIFF(UTexture* SourceImage, bool bWantColor, bool bUseGpu, float Scale, ETextureRotationDirection Rotate)
 {
 	UAppleImageUtilsBaseAsyncTaskBlueprintProxy* Proxy = NewObject<UAppleImageUtilsBaseAsyncTaskBlueprintProxy>();
 	Proxy->SetFlags(RF_StrongRefOnFrame);
 
-	Proxy->ConversionTask = IAppleImageUtilsPlugin::Get().ConvertToTIFF(SourceImage, bWantColor, bUseGpu, Scale);
+	Proxy->ConversionTask = IAppleImageUtilsPlugin::Get().ConvertToTIFF(SourceImage, bWantColor, bUseGpu, Scale, Rotate);
 
 	return Proxy;
 }
 
-UAppleImageUtilsBaseAsyncTaskBlueprintProxy* UAppleImageUtilsBaseAsyncTaskBlueprintProxy::CreateProxyObjectForConvertToPNG(UTexture* SourceImage, bool bWantColor, bool bUseGpu, float Scale)
+UAppleImageUtilsBaseAsyncTaskBlueprintProxy* UAppleImageUtilsBaseAsyncTaskBlueprintProxy::CreateProxyObjectForConvertToPNG(UTexture* SourceImage, bool bWantColor, bool bUseGpu, float Scale, ETextureRotationDirection Rotate)
 {
 	UAppleImageUtilsBaseAsyncTaskBlueprintProxy* Proxy = NewObject<UAppleImageUtilsBaseAsyncTaskBlueprintProxy>();
 	Proxy->SetFlags(RF_StrongRefOnFrame);
 
-	Proxy->ConversionTask = IAppleImageUtilsPlugin::Get().ConvertToPNG(SourceImage, bWantColor, bUseGpu, Scale);
+	Proxy->ConversionTask = IAppleImageUtilsPlugin::Get().ConvertToPNG(SourceImage, bWantColor, bUseGpu, Scale, Rotate);
 
 	return Proxy;
 }
diff --git a/Engine/Plugins/Experimental/AppleImageUtils/Source/AppleImageUtils/Private/AppleImageUtilsPlugin.cpp b/Engine/Plugins/Experimental/AppleImageUtils/Source/AppleImageUtils/Private/AppleImageUtilsPlugin.cpp
--- a/Engine/Plugins/Experimental/AppleImageUtils/Source/AppleImageUtils/Private/AppleImageUtilsPlugin.cpp
+++ b/Engine/Plugins/Experimental/AppleImageUtils/Source/AppleImageUtils/Private/AppleImageUtilsPlugin.cpp
@@ -26,10 +26,10 @@ class FAppleImageUtilsPlugin :
 	virtual void StartupModule() override;
 	virtual void ShutdownModule() override;
 
-	virtual TSharedPtr<FAppleImageUtilsConversionTaskBase, ESPMode::ThreadSafe> ConvertToJPEG(UTexture* SourceImage, int32 Quality = 85, bool bWantColor = true, bool bUseGpu = true, float Scale = 1.f) override;
-	virtual TSharedPtr<FAppleImageUtilsConversionTaskBase, ESPMode::ThreadSafe> ConvertToHEIF(UTexture* SourceImage, int32 Quality = 85, bool bWantColor = true, bool bUseGpu = true, float Scale = 1.f) override;
-	virtual TSharedPtr<FAppleImageUtilsConversionTaskBase, ESPMode::ThreadSafe> ConvertToPNG(UTexture* SourceImage, bool bWantColor = true, bool bUseGpu = true, float Scale = 1.f) override;
-	virtual TSharedPtr<FAppleImageUtilsConversionTaskBase, ESPMode::ThreadSafe> ConvertToTIFF(UTexture* SourceImage, bool bWantColor = true, bool bUseGpu = true, float Scale = 1.f) override;
+	virtual TSharedPtr<FAppleImageUtilsConversionTaskBase, ESPMode::ThreadSafe> ConvertToJPEG(UTexture* SourceImage, int32 Quality = 85, bool bWantColor = true, bool bUseGpu = true, float Scale = 1.f, ETextureRotationDirection Rotate = ETextureRotationDirection::None) override;
+	virtual TSharedPtr<FAppleImageUtilsConversionTaskBase, ESPMode::ThreadSafe> ConvertToHEIF(UTexture* SourceImage, int32 Quality = 85, bool bWantColor = true, bool bUseGpu = true, float Scale = 1.f, ETextureRotationDirection Rotate = ETextureRotationDirection::None) override;
+	virtual TSharedPtr<FAppleImageUtilsConversionTaskBase, ESPMode::ThreadSafe> ConvertToPNG(UTexture* SourceImage, bool bWantColor = true, bool bUseGpu = true, float Scale = 1.f, ETextureRotationDirection Rotate = ETextureRotationDirection::None) override;
+	virtual TSharedPtr<FAppleImageUtilsConversionTaskBase, ESPMode::ThreadSafe> ConvertToTIFF(UTexture* SourceImage, bool bWantColor = true, bool bUseGpu = true, float Scale = 1.f, ETextureRotationDirection Rotate = ETextureRotationDirection::None) override;
 #if SUPPORTS_IMAGE_UTILS_1_0
 	virtual CGImageRef UTexture2DToCGImage(UTexture2D* Source) override;
 #endif
@@ -115,7 +115,7 @@ static inline CGColorSpaceRef ToColorSpace(bool bWantColor)
 	return bWantColor ? CGColorSpaceCreateWithName(kCGColorSpaceSRGB) : CGColorSpaceCreateWithName(kCGColorSpaceGenericGrayGamma2_2);
 }
 
-static inline CIImage* ToImage(IAppleImageInterface* AppleImageInterface, float Scale)
+static inline CIImage* ToImage(IAppleImageInterface* AppleImageInterface, float Scale, ETextureRotationDirection Rotate)
 {
 	check(AppleImageInterface != nullptr);
 
@@ -159,6 +159,19 @@ static inline CIImage* ToImage(IAppleImageInterface* AppleImageInterface, float
 		[ScaleFilter setValue: @(AspectRatio) forKey: kCIInputAspectRatioKey];
 		Image = ScaleFilter.outputImage;
 	}
+	// Handle rotation if requested
+	switch (Rotate)
+	{
+		case ETextureRotationDirection::Left:
+			Image = [Image imageByApplyingOrientation: kCGImagePropertyOrientationLeft];
+			break;
+		case ETextureRotationDirection::Right:
+			Image = [Image imageByApplyingOrientation: kCGImagePropertyOrientationRight];
+			break;
+		case ETextureRotationDirection::Down:
+			Image = [Image imageByApplyingOrientation: kCGImagePropertyOrientationDown];
+			break;
+	}
 	return Image;
 }
 #endif
@@ -177,7 +190,7 @@ static inline bool CanBeConverted(IAppleImageInterface* AppleImageInterface)
 	return false;
 }
 
-TSharedPtr<FAppleImageUtilsConversionTaskBase, ESPMode::ThreadSafe> FAppleImageUtilsPlugin::ConvertToJPEG(UTexture* SourceImage, int32 Quality, bool bWantColor, bool bUseGpu, float Scale)
+TSharedPtr<FAppleImageUtilsConversionTaskBase, ESPMode::ThreadSafe> FAppleImageUtilsPlugin::ConvertToJPEG(UTexture* SourceImage, int32 Quality, bool bWantColor, bool bUseGpu, float Scale, ETextureRotationDirection Rotate)
 {
 	// Make sure our interface is supported
 	IAppleImageInterface* AppleImage = Cast<IAppleImageInterface>(SourceImage);
@@ -197,14 +210,14 @@ TSharedPtr<FAppleImageUtilsConversionTaskBase, ESPMode::ThreadSafe> FAppleImageU
 		return MakeShared<FAppleImageUtilsFailedConversionTask, ESPMode::ThreadSafe>(TEXT("ConvertToJPEG requires iOS 10.0+ or macOS 10.12+"));
 	}
 
-	AsyncTask(ENamedThreads::AnyBackgroundThreadNormalTask, [ConversionTask, Quality, bWantColor, bUseGpu, Scale]()
+	AsyncTask(ENamedThreads::AnyBackgroundThreadNormalTask, [ConversionTask, Quality, bWantColor, bUseGpu, Scale, Rotate]()
 	{
 		IAppleImageInterface* AppleImageInterface = Cast<IAppleImageInterface>(ConversionTask->SourceImage);
 
 		// Convert to Apple objects
 		CIContext* ConversionContext = [[[CIContext alloc] initWithOptions: ToCpuDictionary(bUseGpu)] autorelease];
 		CGColorSpaceRef ColorSpace = ToColorSpace(bWantColor);
-		CIImage* Image = ToImage(AppleImageInterface, Scale);
+		CIImage* Image = ToImage(AppleImageInterface, Scale, Rotate);
 
 		// This will perform the work on the GPU or inline on this thread
 		NSData* ConvertedData = [ConversionContext JPEGRepresentationOfImage: Image colorSpace: ColorSpace options: ToQualityDictionary(Quality)];
@@ -228,7 +241,7 @@ TSharedPtr<FAppleImageUtilsConversionTaskBase, ESPMode::ThreadSafe> FAppleImageU
 	return ConversionTask;
 }
 
-TSharedPtr<FAppleImageUtilsConversionTaskBase, ESPMode::ThreadSafe> FAppleImageUtilsPlugin::ConvertToHEIF(UTexture* SourceImage, int32 Quality, bool bWantColor, bool bUseGpu, float Scale)
+TSharedPtr<FAppleImageUtilsConversionTaskBase, ESPMode::ThreadSafe> FAppleImageUtilsPlugin::ConvertToHEIF(UTexture* SourceImage, int32 Quality, bool bWantColor, bool bUseGpu, float Scale, ETextureRotationDirection Rotate)
 {
 	// Make sure our interface is supported
 	IAppleImageInterface* AppleImage = Cast<IAppleImageInterface>(SourceImage);
@@ -248,14 +261,14 @@ TSharedPtr<FAppleImageUtilsConversionTaskBase, ESPMode::ThreadSafe> FAppleImageU
 		return MakeShared<FAppleImageUtilsFailedConversionTask, ESPMode::ThreadSafe>(TEXT("ConvertToHEIF requires iOS 11.0+ or macOS 10.13.4+"));
 	}
 
-	AsyncTask(ENamedThreads::AnyBackgroundThreadNormalTask, [ConversionTask, Quality, bWantColor, bUseGpu, Scale]()
+	AsyncTask(ENamedThreads::AnyBackgroundThreadNormalTask, [ConversionTask, Quality, bWantColor, bUseGpu, Scale, Rotate]()
 	{
 		IAppleImageInterface* AppleImageInterface = Cast<IAppleImageInterface>(ConversionTask->SourceImage);
 
 		// Convert to Apple objects
 		CIContext* ConversionContext = [[[CIContext alloc] initWithOptions: ToCpuDictionary(bUseGpu)] autorelease];
 		CGColorSpaceRef ColorSpace = ToColorSpace(bWantColor);
-		CIImage* Image = ToImage(AppleImageInterface, Scale);
+		CIImage* Image = ToImage(AppleImageInterface, Scale, Rotate);
 
 		// This will perform the work on the GPU or inline on this thread
 		NSData* ConvertedData = [ConversionContext HEIFRepresentationOfImage: Image format: kCIFormatARGB8 colorSpace: ColorSpace options: ToQualityDictionary(Quality)];
@@ -279,7 +292,7 @@ TSharedPtr<FAppleImageUtilsConversionTaskBase, ESPMode::ThreadSafe> FAppleImageU
 	return ConversionTask;
 }
 
-TSharedPtr<FAppleImageUtilsConversionTaskBase, ESPMode::ThreadSafe> FAppleImageUtilsPlugin::ConvertToPNG(UTexture* SourceImage, bool bWantColor, bool bUseGpu, float Scale)
+TSharedPtr<FAppleImageUtilsConversionTaskBase, ESPMode::ThreadSafe> FAppleImageUtilsPlugin::ConvertToPNG(UTexture* SourceImage, bool bWantColor, bool bUseGpu, float Scale, ETextureRotationDirection Rotate)
 {
 	// Make sure our interface is supported
 	IAppleImageInterface* AppleImage = Cast<IAppleImageInterface>(SourceImage);
@@ -299,14 +312,14 @@ TSharedPtr<FAppleImageUtilsConversionTaskBase, ESPMode::ThreadSafe> FAppleImageU
 		return MakeShared<FAppleImageUtilsFailedConversionTask, ESPMode::ThreadSafe>(TEXT("ConvertToPNG requires iOS 11.0+ or macOS 10.13+"));
 	}
 
-	AsyncTask(ENamedThreads::AnyBackgroundThreadNormalTask, [ConversionTask, bWantColor, bUseGpu, Scale]()
+	AsyncTask(ENamedThreads::AnyBackgroundThreadNormalTask, [ConversionTask, bWantColor, bUseGpu, Scale, Rotate]()
 	{
 		IAppleImageInterface* AppleImageInterface = Cast<IAppleImageInterface>(ConversionTask->SourceImage);
 
 		// Convert to Apple objects
 		CIContext* ConversionContext = [[[CIContext alloc] initWithOptions: ToCpuDictionary(bUseGpu)] autorelease];
 		CGColorSpaceRef ColorSpace = ToColorSpace(bWantColor);
-		CIImage* Image = ToImage(AppleImageInterface, Scale);
+		CIImage* Image = ToImage(AppleImageInterface, Scale, Rotate);
 
 		// This will perform the work on the GPU or inline on this thread
 		NSData* ConvertedData = [ConversionContext PNGRepresentationOfImage: Image format: kCIFormatARGB8 colorSpace: ColorSpace options: @{}];
@@ -330,7 +343,7 @@ TSharedPtr<FAppleImageUtilsConversionTaskBase, ESPMode::ThreadSafe> FAppleImageU
 	return ConversionTask;
 }
 
-TSharedPtr<FAppleImageUtilsConversionTaskBase, ESPMode::ThreadSafe> FAppleImageUtilsPlugin::ConvertToTIFF(UTexture* SourceImage, bool bWantColor, bool bUseGpu, float Scale)
+TSharedPtr<FAppleImageUtilsConversionTaskBase, ESPMode::ThreadSafe> FAppleImageUtilsPlugin::ConvertToTIFF(UTexture* SourceImage, bool bWantColor, bool bUseGpu, float Scale, ETextureRotationDirection Rotate)
 {
 	// Make sure our interface is supported
 	IAppleImageInterface* AppleImage = Cast<IAppleImageInterface>(SourceImage);
@@ -350,14 +363,14 @@ TSharedPtr<FAppleImageUtilsConversionTaskBase, ESPMode::ThreadSafe> FAppleImageU
 		return MakeShared<FAppleImageUtilsFailedConversionTask, ESPMode::ThreadSafe>(TEXT("ConvertToTIFF requires iOS 10.0+ or macOS 10.12+"));
 	}
 
-	AsyncTask(ENamedThreads::AnyBackgroundThreadNormalTask, [ConversionTask, bWantColor, bUseGpu, Scale]()
+	AsyncTask(ENamedThreads::AnyBackgroundThreadNormalTask, [ConversionTask, bWantColor, bUseGpu, Scale, Rotate]()
 	{
 		IAppleImageInterface* AppleImageInterface = Cast<IAppleImageInterface>(ConversionTask->SourceImage);
 
 		// Convert to Apple objects
 		CIContext* ConversionContext = [[[CIContext alloc] initWithOptions: ToCpuDictionary(bUseGpu)] autorelease];
 		CGColorSpaceRef ColorSpace = ToColorSpace(bWantColor);
-		CIImage* Image = ToImage(AppleImageInterface, Scale);
+		CIImage* Image = ToImage(AppleImageInterface, Scale, Rotate);
 
 		// This will perform the work on the GPU or inline on this thread
 		NSData* ConvertedData = [ConversionContext TIFFRepresentationOfImage: Image format: kCIFormatARGB8 colorSpace: ColorSpace options: @{}];
diff --git a/Engine/Plugins/Experimental/AppleImageUtils/Source/AppleImageUtils/Public/AppleImageUtilsBlueprintProxy.h b/Engine/Plugins/Experimental/AppleImageUtils/Source/AppleImageUtils/Public/AppleImageUtilsBlueprintProxy.h
--- a/Engine/Plugins/Experimental/AppleImageUtils/Source/AppleImageUtils/Public/AppleImageUtilsBlueprintProxy.h
+++ b/Engine/Plugins/Experimental/AppleImageUtils/Source/AppleImageUtils/Public/AppleImageUtilsBlueprintProxy.h
@@ -50,9 +50,10 @@ public:
 	 * @param bWantColor whether the JPEG is color (true) or monochrome (false)
 	 * @param bUseGpu whether to use the GPU (true) or the CPU (false) to compress
 	 * @param Scale whether to scale the image before conversion, defaults to no scaling
+	 * @param Rotate a direction to rotate the image in during conversion, defaults to none
 	 */
 	UFUNCTION(BlueprintCallable, Meta=(BlueprintInternalUseOnly="true", DisplayName="Convert To JPEG"), Category="Image Conversion")
-	static UAppleImageUtilsBaseAsyncTaskBlueprintProxy* CreateProxyObjectForConvertToJPEG(UTexture* SourceImage, int32 Quality = 85, bool bWantColor = true, bool bUseGpu = true, float Scale = 1.f);
+	static UAppleImageUtilsBaseAsyncTaskBlueprintProxy* CreateProxyObjectForConvertToJPEG(UTexture* SourceImage, int32 Quality = 85, bool bWantColor = true, bool bUseGpu = true, float Scale = 1.f, ETextureRotationDirection Rotate = ETextureRotationDirection::None);
 
 	/**
 	 * Converts a image to an array of HEIF data in a background task
@@ -62,9 +63,10 @@ public:
 	 * @param bWantColor whether the HEIF is color (true) or monochrome (false)
 	 * @param bUseGpu whether to use the GPU (true) or the CPU (false) to compress
 	 * @param Scale whether to scale the image before conversion, defaults to no scaling
+	 * @param Rotate a direction to rotate the image in during conversion, defaults to none
 	 */
 	UFUNCTION(BlueprintCallable, Meta=(BlueprintInternalUseOnly="true", DisplayName="Convert To HEIF"), Category="Image Conversion")
-	static UAppleImageUtilsBaseAsyncTaskBlueprintProxy* CreateProxyObjectForConvertToHEIF(UTexture* SourceImage, int32 Quality = 85, bool bWantColor = true, bool bUseGpu = true, float Scale = 1.f);
+	static UAppleImageUtilsBaseAsyncTaskBlueprintProxy* CreateProxyObjectForConvertToHEIF(UTexture* SourceImage, int32 Quality = 85, bool bWantColor = true, bool bUseGpu = true, float Scale = 1.f, ETextureRotationDirection Rotate = ETextureRotationDirection::None);
 
 	/**
 	 * Converts a image to an array of TIFF data in a background task
@@ -74,9 +76,10 @@ public:
 	 * @param bWantColor whether the TIFF is color (true) or monochrome (false)
 	 * @param bUseGpu whether to use the GPU (true) or the CPU (false) to compress
 	 * @param Scale whether to scale the image before conversion, defaults to no scaling
+	 * @param Rotate a direction to rotate the image in during conversion, defaults to none
 	 */
 	UFUNCTION(BlueprintCallable, Meta=(BlueprintInternalUseOnly="true", DisplayName="Convert To TIFF"), Category="Image Conversion")
-	static UAppleImageUtilsBaseAsyncTaskBlueprintProxy* CreateProxyObjectForConvertToTIFF(UTexture* SourceImage, bool bWantColor = true, bool bUseGpu = true, float Scale = 1.f);
+	static UAppleImageUtilsBaseAsyncTaskBlueprintProxy* CreateProxyObjectForConvertToTIFF(UTexture* SourceImage, bool bWantColor = true, bool bUseGpu = true, float Scale = 1.f, ETextureRotationDirection Rotate = ETextureRotationDirection::None);
 
 	/**
 	 * Converts a image to an array of TIFF data in a background task
@@ -86,9 +89,10 @@ public:
 	 * @param bWantColor whether the PNG is color (true) or monochrome (false)
 	 * @param bUseGpu whether to use the GPU (true) or the CPU (false) to compress
 	 * @param Scale whether to scale the image before conversion, defaults to no scaling
+	 * @param Rotate a direction to rotate the image in during conversion, defaults to none
 	 */
 	UFUNCTION(BlueprintCallable, Meta=(BlueprintInternalUseOnly="true", DisplayName="Convert To PNG"), Category="Image Conversion")
-	static UAppleImageUtilsBaseAsyncTaskBlueprintProxy* CreateProxyObjectForConvertToPNG(UTexture* SourceImage, bool bWantColor = true, bool bUseGpu = true, float Scale = 1.f);
+	static UAppleImageUtilsBaseAsyncTaskBlueprintProxy* CreateProxyObjectForConvertToPNG(UTexture* SourceImage, bool bWantColor = true, bool bUseGpu = true, float Scale = 1.f, ETextureRotationDirection Rotate = ETextureRotationDirection::None);
 
 	//~ Begin FTickableObject Interface
 	virtual void Tick(float DeltaTime) override;
diff --git a/Engine/Plugins/Experimental/AppleImageUtils/Source/AppleImageUtils/Public/AppleImageUtilsTypes.h b/Engine/Plugins/Experimental/AppleImageUtils/Source/AppleImageUtils/Public/AppleImageUtilsTypes.h
--- a/Engine/Plugins/Experimental/AppleImageUtils/Source/AppleImageUtils/Public/AppleImageUtilsTypes.h
+++ b/Engine/Plugins/Experimental/AppleImageUtils/Source/AppleImageUtils/Public/AppleImageUtilsTypes.h
@@ -11,6 +11,15 @@
 #endif
 #include "AppleImageUtilsTypes.generated.h"
 
+UENUM(BlueprintType, Category="AppleImageUtils", meta=(Experimental))
+enum class ETextureRotationDirection : uint8
+{
+	None,
+	Left,
+	Right,
+	Down
+};
+
 UENUM(BlueprintType, Category="AppleImageUtils", meta=(Experimental))
 enum class EAppleTextureType : uint8
 {
diff --git a/Engine/Plugins/Experimental/AppleImageUtils/Source/AppleImageUtils/Public/IAppleImageUtilsPlugin.h b/Engine/Plugins/Experimental/AppleImageUtils/Source/AppleImageUtils/Public/IAppleImageUtilsPlugin.h
--- a/Engine/Plugins/Experimental/AppleImageUtils/Source/AppleImageUtils/Public/IAppleImageUtilsPlugin.h
+++ b/Engine/Plugins/Experimental/AppleImageUtils/Source/AppleImageUtils/Public/IAppleImageUtilsPlugin.h
@@ -7,6 +7,8 @@
 #include "Modules/ModuleManager.h"
 #include "HAL/ThreadSafeBool.h"
 
+#include "AppleImageUtilsTypes.h"
+
 #include "AppleImageUtilsAvailability.h"
 
 #if SUPPORTS_IMAGE_UTILS_1_0
@@ -86,10 +88,11 @@ public:
 	 * @param bWantColor whether the JPEG is color (true) or monochrome (false)
 	 * @param bUseGpu whether to use the GPU (true) or the CPU (false) to compress
 	 * @param Scale whether to scale the image before converting, defaults to no scale operation
+	 * @param Rotate a direction to rotate the image in during conversion, defaults to none
 	 *
 	 * @return the async task that is doing the conversion
 	 */
-	virtual TSharedPtr<FAppleImageUtilsConversionTaskBase, ESPMode::ThreadSafe> ConvertToJPEG(UTexture* SourceImage, int32 Quality = 85, bool bWantColor = true, bool bUseGpu = true, float Scale = 1.f) = 0;
+	virtual TSharedPtr<FAppleImageUtilsConversionTaskBase, ESPMode::ThreadSafe> ConvertToJPEG(UTexture* SourceImage, int32 Quality = 85, bool bWantColor = true, bool bUseGpu = true, float Scale = 1.f, ETextureRotationDirection Rotate = ETextureRotationDirection::None) = 0;
 
 	/**
 	 * Converts a image to an array of HEIF data in a background task
@@ -98,10 +101,12 @@ public:
 	 * @param Quality the quality level to compress to
 	 * @param bWantColor whether the HEIF is color (true) or monochrome (false)
 	 * @param bUseGpu whether to use the GPU (true) or the CPU (false) to compress
+	 * @param Scale whether to scale the image before converting, defaults to no scale operation
+	 * @param Rotate a direction to rotate the image in during conversion, defaults to none
 	 *
 	 * @return the async task that is doing the conversion
 	 */
-	virtual TSharedPtr<FAppleImageUtilsConversionTaskBase, ESPMode::ThreadSafe> ConvertToHEIF(UTexture* SourceImage, int32 Quality = 85,  bool bWantColor = true, bool bUseGpu = true, float Scale = 1.f) = 0;
+	virtual TSharedPtr<FAppleImageUtilsConversionTaskBase, ESPMode::ThreadSafe> ConvertToHEIF(UTexture* SourceImage, int32 Quality = 85,  bool bWantColor = true, bool bUseGpu = true, float Scale = 1.f, ETextureRotationDirection Rotate = ETextureRotationDirection::None) = 0;
 
 	/**
 	 * Converts a image to an array of PNG data in a background task
@@ -110,10 +115,12 @@ public:
 	 * @param Quality the quality level to compress to
 	 * @param bWantColor whether the PNG is color (true) or monochrome (false)
 	 * @param bUseGpu whether to use the GPU (true) or the CPU (false) to compress
+	 * @param Scale whether to scale the image before converting, defaults to no scale operation
+	 * @param Rotate a direction to rotate the image in during conversion, defaults to none
 	 *
 	 * @return the async task that is doing the conversion
 	 */
-	virtual TSharedPtr<FAppleImageUtilsConversionTaskBase, ESPMode::ThreadSafe> ConvertToPNG(UTexture* SourceImage, bool bWantColor = true, bool bUseGpu = true, float Scale = 1.f) = 0;
+	virtual TSharedPtr<FAppleImageUtilsConversionTaskBase, ESPMode::ThreadSafe> ConvertToPNG(UTexture* SourceImage, bool bWantColor = true, bool bUseGpu = true, float Scale = 1.f, ETextureRotationDirection Rotate = ETextureRotationDirection::None) = 0;
 
 	/**
 	 * Converts a image to an array of TIFF data in a background task
@@ -121,10 +128,12 @@ public:
 	 * @param SourceImage the image to compress (NOTE: must support UAppleImageInterface)
 	 * @param Quality the quality level to compress to
 	 * @param bUseGpu whether to use the GPU (true) or the CPU (false) to compress
+	 * @param Scale whether to scale the image before converting, defaults to no scale operation
+	 * @param Rotate a direction to rotate the image in during conversion, defaults to none
 	 *
 	 * @return the async task that is doing the conversion
 	 */
-	virtual TSharedPtr<FAppleImageUtilsConversionTaskBase, ESPMode::ThreadSafe> ConvertToTIFF(UTexture* SourceImage, bool bWantColor = true, bool bUseGpu = true, float Scale = 1.f) = 0;
+	virtual TSharedPtr<FAppleImageUtilsConversionTaskBase, ESPMode::ThreadSafe> ConvertToTIFF(UTexture* SourceImage, bool bWantColor = true, bool bUseGpu = true, float Scale = 1.f, ETextureRotationDirection Rotate = ETextureRotationDirection::None) = 0;
 
 #if SUPPORTS_IMAGE_UTILS_1_0
 	/**
