Hid blueprint internal methods from being callable

diff --git a/Engine/Plugins/Experimental/AppleImageUtils/Source/AppleImageUtils/Public/AppleImageUtilsBlueprintProxy.h b/Engine/Plugins/Experimental/AppleImageUtils/Source/AppleImageUtils/Public/AppleImageUtilsBlueprintProxy.h
--- a/Engine/Plugins/Experimental/AppleImageUtils/Source/AppleImageUtils/Public/AppleImageUtilsBlueprintProxy.h
+++ b/Engine/Plugins/Experimental/AppleImageUtils/Source/AppleImageUtils/Public/AppleImageUtilsBlueprintProxy.h
@@ -50,7 +50,7 @@ public:
 	 * @param bWantColor whether the JPEG is color (true) or monochrome (false)
 	 * @param bUseGpu whether to use the GPU (true) or the CPU (false) to compress
 	 */
-	UFUNCTION(BlueprintCallable, Meta=(DisplayName="Convert To JPEG"), Category="Image Conversion")
+	UFUNCTION(BlueprintCallable, Meta=(BlueprintInternalUseOnly="true", DisplayName="Convert To JPEG"), Category="Image Conversion")
 	static UAppleImageUtilsBaseAsyncTaskBlueprintProxy* CreateProxyObjectForConvertToJPEG(UTexture* SourceImage, int32 Quality = 85, bool bWantColor = true, bool bUseGpu = true);
 
 	/**
@@ -61,7 +61,7 @@ public:
 	 * @param bWantColor whether the HEIF is color (true) or monochrome (false)
 	 * @param bUseGpu whether to use the GPU (true) or the CPU (false) to compress
 	 */
-	UFUNCTION(BlueprintCallable, Meta=(DisplayName="Convert To HEIF"), Category="Image Conversion")
+	UFUNCTION(BlueprintCallable, Meta=(BlueprintInternalUseOnly="true", DisplayName="Convert To HEIF"), Category="Image Conversion")
 	static UAppleImageUtilsBaseAsyncTaskBlueprintProxy* CreateProxyObjectForConvertToHEIF(UTexture* SourceImage, int32 Quality = 85, bool bWantColor = true, bool bUseGpu = true);
 
 	/**
@@ -72,7 +72,7 @@ public:
 	 * @param bWantColor whether the TIFF is color (true) or monochrome (false)
 	 * @param bUseGpu whether to use the GPU (true) or the CPU (false) to compress
 	 */
-	UFUNCTION(BlueprintCallable, Meta=(DisplayName="Convert To TIFF"), Category="Image Conversion")
+	UFUNCTION(BlueprintCallable, Meta=(BlueprintInternalUseOnly="true", DisplayName="Convert To TIFF"), Category="Image Conversion")
 	static UAppleImageUtilsBaseAsyncTaskBlueprintProxy* CreateProxyObjectForConvertToTIFF(UTexture* SourceImage, bool bWantColor = true, bool bUseGpu = true);
 
 	/**
@@ -83,7 +83,7 @@ public:
 	 * @param bWantColor whether the PNG is color (true) or monochrome (false)
 	 * @param bUseGpu whether to use the GPU (true) or the CPU (false) to compress
 	 */
-	UFUNCTION(BlueprintCallable, Meta=(DisplayName="Convert To PNG"), Category="Image Conversion")
+	UFUNCTION(BlueprintCallable, Meta=(BlueprintInternalUseOnly="true", DisplayName="Convert To PNG"), Category="Image Conversion")
 	static UAppleImageUtilsBaseAsyncTaskBlueprintProxy* CreateProxyObjectForConvertToPNG(UTexture* SourceImage, bool bWantColor = true, bool bUseGpu = true);
 
 	//~ Begin FTickableObject Interface
