Fixed wrong platform name being used as the whitelist for the AppleImageUtils plugin

diff --git a/Engine/Plugins/Experimental/AppleImageUtils/AppleImageUtils.uplugin b/Engine/Plugins/Experimental/AppleImageUtils/AppleImageUtils.uplugin
--- a/Engine/Plugins/Experimental/AppleImageUtils/AppleImageUtils.uplugin
+++ b/Engine/Plugins/Experimental/AppleImageUtils/AppleImageUtils.uplugin
@@ -31,10 +31,10 @@
 			"LoadingPhase" : "Default",
 			"WhitelistPlatforms" :
 			[
-				"Windows",
+				"Win64",
 				"Mac",
 				"Linux"
 			]
 		}
 	]
-}
\ No newline at end of file
+}
