Added missing module class for AppleImageUtilsBlueprintSupport

diff --git a/Engine/Plugins/Experimental/AppleImageUtils/Source/AppleImageUtilsBlueprintSupport/Private/AppleImageUtilsBlueprintSupport.cpp b/Engine/Plugins/Experimental/AppleImageUtils/Source/AppleImageUtilsBlueprintSupport/Private/AppleImageUtilsBlueprintSupport.cpp
--- a/Engine/Plugins/Experimental/AppleImageUtils/Source/AppleImageUtilsBlueprintSupport/Private/AppleImageUtilsBlueprintSupport.cpp
+++ b/Engine/Plugins/Experimental/AppleImageUtils/Source/AppleImageUtilsBlueprintSupport/Private/AppleImageUtilsBlueprintSupport.cpp
@@ -2,6 +2,20 @@
 
 #include "AppleImageUtilsBlueprintSupport.h"
 #include "AppleImageUtilsBlueprintProxy.h"
+#include "Modules/ModuleInterface.h"
+#include "Modules/ModuleManager.h"
+
+
+class FAppleImageUtilsBlueprintSupport :
+	public IModuleInterface
+{
+	/** IModuleInterface implementation */
+	virtual void StartupModule() override {}
+	virtual void ShutdownModule() override {}
+
+};
+
+IMPLEMENT_MODULE(FAppleImageUtilsBlueprintSupport, AppleImageUtilsBlueprintSupport);
 
 UK2Node_ConvertToJPEG::UK2Node_ConvertToJPEG(const FObjectInitializer& ObjectInitializer)
 	: Super(ObjectInitializer)
