Changed how UAppleARKitTextureCameraImage are used when wrapping the underlying AR image (reuse the object)
Fixed up some autorelease pool problems in the apple image utils exposed by fixing lifetimes

diff --git a/Engine/Plugins/Experimental/AppleImageUtils/Source/AppleImageUtils/Private/AppleImageUtilsPlugin.cpp b/Engine/Plugins/Experimental/AppleImageUtils/Source/AppleImageUtils/Private/AppleImageUtilsPlugin.cpp
--- a/Engine/Plugins/Experimental/AppleImageUtils/Source/AppleImageUtils/Private/AppleImageUtilsPlugin.cpp
+++ b/Engine/Plugins/Experimental/AppleImageUtils/Source/AppleImageUtils/Private/AppleImageUtilsPlugin.cpp
@@ -9,7 +9,7 @@
 #include "Async/Async.h"
 #include "Engine/Texture.h"
 #include "Engine/Texture2D.h"
-
+#include "Apple/ApplePlatformMisc.h"
 
 #if SUPPORTS_IMAGE_UTILS_1_0
 	#import <CoreImage/CIContext.h>
@@ -99,14 +99,14 @@ static inline NSDictionary* ToQualityDictionary(int32 Quality)
 	Quality = FMath::Clamp(Quality, 0, 100);
 	float QualityF = float(Quality) * 0.01f;
 
-	NSDictionary* Options = [[NSDictionary dictionaryWithObjectsAndKeys:[NSNumber numberWithFloat: QualityF], kCGImageDestinationLossyCompressionQuality, nil] autorelease];
+	NSDictionary* Options = [NSDictionary dictionaryWithObjectsAndKeys:[NSNumber numberWithFloat: QualityF], kCGImageDestinationLossyCompressionQuality, nil];
 	return Options;
 }
 
 static inline NSDictionary* ToCpuDictionary(bool bUseGpu)
 {
 	BOOL UseCpu = bUseGpu ? NO : YES;
-	NSDictionary* Options = [[NSDictionary dictionaryWithObjectsAndKeys:[NSNumber numberWithBool: UseCpu], kCIContextUseSoftwareRenderer, nil] autorelease];
+	NSDictionary* Options = [NSDictionary dictionaryWithObjectsAndKeys:[NSNumber numberWithBool: UseCpu], kCIContextUseSoftwareRenderer, nil];
 	return Options;
 }
 
@@ -133,7 +133,7 @@ static inline CIImage* ToImage(IAppleImageInterface* AppleImageInterface, float
 			CVPixelBufferRef PixelBuffer = AppleImageInterface->GetPixelBuffer();
 			if (PixelBuffer != nullptr)
 			{
-				Image = [[CIImage imageWithCVPixelBuffer: PixelBuffer] autorelease];
+				Image = [CIImage imageWithCVPixelBuffer: PixelBuffer];
 			}
 			break;
 		}
@@ -143,7 +143,7 @@ static inline CIImage* ToImage(IAppleImageInterface* AppleImageInterface, float
 			IOSurfaceRef Surface = AppleImageInterface->GetSurface();
 			if (Surface != nullptr)
 			{
-				Image = [[CIImage imageWithIOSurface: Surface] autorelease];
+				Image = [CIImage imageWithIOSurface: Surface];
 			}
 			break;
 		}
@@ -212,10 +212,12 @@ TSharedPtr<FAppleImageUtilsConversionTaskBase, ESPMode::ThreadSafe> FAppleImageU
 
 	AsyncTask(ENamedThreads::AnyBackgroundThreadNormalTask, [ConversionTask, Quality, bWantColor, bUseGpu, Scale, Rotate]()
 	{
+		SCOPED_AUTORELEASE_POOL;
+
 		IAppleImageInterface* AppleImageInterface = Cast<IAppleImageInterface>(ConversionTask->SourceImage);
 
 		// Convert to Apple objects
-		CIContext* ConversionContext = [[[CIContext alloc] initWithOptions: ToCpuDictionary(bUseGpu)] autorelease];
+		CIContext* ConversionContext = [CIContext contextWithOptions: ToCpuDictionary(bUseGpu)];
 		CGColorSpaceRef ColorSpace = ToColorSpace(bWantColor);
 		CIImage* Image = ToImage(AppleImageInterface, Scale, Rotate);
 
@@ -226,12 +228,10 @@ TSharedPtr<FAppleImageUtilsConversionTaskBase, ESPMode::ThreadSafe> FAppleImageU
 			uint32 CompressedSize = ConvertedData.length;
 			ConversionTask->ConvertedBytes.AddUninitialized(CompressedSize);
 			FPlatformMemory::Memcpy(ConversionTask->ConvertedBytes.GetData(), [ConvertedData bytes], CompressedSize);
-
-			[ConvertedData release];
 		}
 
 		CGColorSpaceRelease(ColorSpace);
-
+		
 		// Notify any async listeners that we are done
 		ConversionTask->MarkComplete();
 	});
@@ -263,10 +263,12 @@ TSharedPtr<FAppleImageUtilsConversionTaskBase, ESPMode::ThreadSafe> FAppleImageU
 
 	AsyncTask(ENamedThreads::AnyBackgroundThreadNormalTask, [ConversionTask, Quality, bWantColor, bUseGpu, Scale, Rotate]()
 	{
+		SCOPED_AUTORELEASE_POOL;
+
 		IAppleImageInterface* AppleImageInterface = Cast<IAppleImageInterface>(ConversionTask->SourceImage);
 
 		// Convert to Apple objects
-		CIContext* ConversionContext = [[[CIContext alloc] initWithOptions: ToCpuDictionary(bUseGpu)] autorelease];
+		CIContext* ConversionContext = [CIContext contextWithOptions: ToCpuDictionary(bUseGpu)];
 		CGColorSpaceRef ColorSpace = ToColorSpace(bWantColor);
 		CIImage* Image = ToImage(AppleImageInterface, Scale, Rotate);
 
@@ -277,8 +279,6 @@ TSharedPtr<FAppleImageUtilsConversionTaskBase, ESPMode::ThreadSafe> FAppleImageU
 			uint32 CompressedSize = ConvertedData.length;
 			ConversionTask->ConvertedBytes.AddUninitialized(CompressedSize);
 			FPlatformMemory::Memcpy(ConversionTask->ConvertedBytes.GetData(), [ConvertedData bytes], CompressedSize);
-
-			[ConvertedData release];
 		}
 
 		CGColorSpaceRelease(ColorSpace);
@@ -314,10 +314,12 @@ TSharedPtr<FAppleImageUtilsConversionTaskBase, ESPMode::ThreadSafe> FAppleImageU
 
 	AsyncTask(ENamedThreads::AnyBackgroundThreadNormalTask, [ConversionTask, bWantColor, bUseGpu, Scale, Rotate]()
 	{
+		SCOPED_AUTORELEASE_POOL;
+
 		IAppleImageInterface* AppleImageInterface = Cast<IAppleImageInterface>(ConversionTask->SourceImage);
 
 		// Convert to Apple objects
-		CIContext* ConversionContext = [[[CIContext alloc] initWithOptions: ToCpuDictionary(bUseGpu)] autorelease];
+		CIContext* ConversionContext = [CIContext contextWithOptions: ToCpuDictionary(bUseGpu)];
 		CGColorSpaceRef ColorSpace = ToColorSpace(bWantColor);
 		CIImage* Image = ToImage(AppleImageInterface, Scale, Rotate);
 
@@ -328,8 +330,6 @@ TSharedPtr<FAppleImageUtilsConversionTaskBase, ESPMode::ThreadSafe> FAppleImageU
 			uint32 CompressedSize = ConvertedData.length;
 			ConversionTask->ConvertedBytes.AddUninitialized(CompressedSize);
 			FPlatformMemory::Memcpy(ConversionTask->ConvertedBytes.GetData(), [ConvertedData bytes], CompressedSize);
-
-			[ConvertedData release];
 		}
 
 		CGColorSpaceRelease(ColorSpace);
@@ -365,10 +365,12 @@ TSharedPtr<FAppleImageUtilsConversionTaskBase, ESPMode::ThreadSafe> FAppleImageU
 
 	AsyncTask(ENamedThreads::AnyBackgroundThreadNormalTask, [ConversionTask, bWantColor, bUseGpu, Scale, Rotate]()
 	{
+		SCOPED_AUTORELEASE_POOL;
+
 		IAppleImageInterface* AppleImageInterface = Cast<IAppleImageInterface>(ConversionTask->SourceImage);
 
 		// Convert to Apple objects
-		CIContext* ConversionContext = [[[CIContext alloc] initWithOptions: ToCpuDictionary(bUseGpu)] autorelease];
+		CIContext* ConversionContext = [CIContext contextWithOptions: ToCpuDictionary(bUseGpu)];
 		CGColorSpaceRef ColorSpace = ToColorSpace(bWantColor);
 		CIImage* Image = ToImage(AppleImageInterface, Scale, Rotate);
 
@@ -379,8 +381,6 @@ TSharedPtr<FAppleImageUtilsConversionTaskBase, ESPMode::ThreadSafe> FAppleImageU
 			uint32 CompressedSize = ConvertedData.length;
 			ConversionTask->ConvertedBytes.AddUninitialized(CompressedSize);
 			FPlatformMemory::Memcpy(ConversionTask->ConvertedBytes.GetData(), [ConvertedData bytes], CompressedSize);
-
-			[ConvertedData release];
 		}
 
 		CGColorSpaceRelease(ColorSpace);
diff --git a/Engine/Plugins/Runtime/AppleARKit/Source/AppleARKit/Private/AppleARKitFrame.cpp b/Engine/Plugins/Runtime/AppleARKit/Source/AppleARKit/Private/AppleARKitFrame.cpp
--- a/Engine/Plugins/Runtime/AppleARKit/Source/AppleARKit/Private/AppleARKitFrame.cpp
+++ b/Engine/Plugins/Runtime/AppleARKit/Source/AppleARKit/Private/AppleARKitFrame.cpp
@@ -72,6 +72,7 @@ FAppleARKitFrame::FAppleARKitFrame( ARFrame* InARFrame, CVMetalTextureCacheRef M
 	{
 		CameraImage = InARFrame.capturedImage;
 		CFRetain(CameraImage);
+
 		// Update SizeX & Y
 		CapturedYImageWidth = CVPixelBufferGetWidthOfPlane( InARFrame.capturedImage, 0 );
 		CapturedYImageHeight = CVPixelBufferGetHeightOfPlane( InARFrame.capturedImage, 0 );
@@ -93,9 +94,8 @@ FAppleARKitFrame::FAppleARKitFrame( ARFrame* InARFrame, CVMetalTextureCacheRef M
 		check( CapturedCbCrImage );
 		check( CFGetRetainCount(CapturedCbCrImage) == 1);
 	}
-	// @todo JoeG -- finsih the depth capture
-//@joeg -- Disabled due to crashing when accessing
-	if (0 && InARFrame.capturedDepthData)
+
+	if (InARFrame.capturedDepthData)
 	{
 		CameraDepth = InARFrame.capturedDepthData;
 		CFRetain(CameraDepth);
@@ -172,7 +172,16 @@ FAppleARKitFrame& FAppleARKitFrame::operator=( const FAppleARKitFrame& Other )
 	{
 		CFRelease( CapturedCbCrImage );
 	}
-
+	if (CameraImage != nullptr)
+	{
+		CFRelease(CameraImage);
+		CameraImage = nullptr;
+	}
+	if (CameraDepth != nullptr)
+	{
+		CFRelease(CameraDepth);
+		CameraDepth = nullptr;
+	}
 	if(NativeFrame != nullptr)
 	{
 		CFRelease((CFTypeRef)NativeFrame);
diff --git a/Engine/Plugins/Runtime/AppleARKit/Source/AppleARKit/Private/AppleARKitSystem.cpp b/Engine/Plugins/Runtime/AppleARKit/Source/AppleARKit/Private/AppleARKitSystem.cpp
--- a/Engine/Plugins/Runtime/AppleARKit/Source/AppleARKit/Private/AppleARKitSystem.cpp
+++ b/Engine/Plugins/Runtime/AppleARKit/Source/AppleARKit/Private/AppleARKitSystem.cpp
@@ -364,25 +364,25 @@ void FAppleARKitSystem::UpdateFrame()
 			GameThreadTimestamp = GameThreadFrame->Timestamp;
 
 #if SUPPORTS_ARKIT_1_0
-			// Only create a new camera image texture if it's set
 			if (GameThreadFrame->CameraImage != nullptr)
 			{
-				if (!bCanReuseCameraImage || CameraImage == nullptr)
+				// Only create a new camera image texture if it's set and we don't already have one
+				if (CameraImage == nullptr)
 				{
 					CameraImage = NewObject<UAppleARKitTextureCameraImage>();
-					bCanReuseCameraImage = true;
 				}
+				// Reuse the UObjects because otherwise the time between GCs causes ARKit to be starved of resources
                 CameraImage->Init(GameThreadTimestamp, GameThreadFrame->CameraImage);
 			}
 
-			// Only create a new camera depth texture if it's set (currently only Face AR)
 			if (GameThreadFrame->CameraDepth != nullptr)
 			{
-				if (!bCanReuseCameraDepth || CameraDepth == nullptr)
+				// Only create a new camera depth texture if it's set and we don't already have one
+				if (CameraDepth == nullptr)
 				{
 					CameraDepth = NewObject<UAppleARKitTextureCameraDepth>();
-					bCanReuseCameraDepth = true;
 				}
+				// Reuse the UObjects because otherwise the time between GCs causes ARKit to be starved of resources
                 CameraDepth->Init(GameThreadTimestamp, GameThreadFrame->CameraDepth);
 			}
 #endif
@@ -700,13 +700,11 @@ TArray<UARPin*> FAppleARKitSystem::OnGetAllPins() const
 
 UARTextureCameraImage* FAppleARKitSystem::OnGetCameraImage()
 {
-	bCanReuseCameraImage = false;
 	return CameraImage;
 }
 
 UARTextureCameraDepth* FAppleARKitSystem::OnGetCameraDepth()
 {
-	bCanReuseCameraDepth = false;
 	return CameraDepth;
 }
 
diff --git a/Engine/Plugins/Runtime/AppleARKit/Source/AppleARKit/Private/AppleARKitSystem.h b/Engine/Plugins/Runtime/AppleARKit/Source/AppleARKit/Private/AppleARKitSystem.h
--- a/Engine/Plugins/Runtime/AppleARKit/Source/AppleARKit/Private/AppleARKitSystem.h
+++ b/Engine/Plugins/Runtime/AppleARKit/Source/AppleARKit/Private/AppleARKitSystem.h
@@ -175,11 +175,6 @@ private:
 	// PROPERTIES REPORTED TO FGCObject
 	//
 	
-	/** Controls whether we can use the existing wrapper object instead of creating a new one to cut down on UObject churn */
-	bool bCanReuseCameraImage;
-	/** Controls whether we can use the existing wrapper object instead of creating a new one to cut down on UObject churn */
-	bool bCanReuseCameraDepth;
-
 	/** The ar frame number when LastReceivedFrame was last updated */
 	uint32 GameThreadFrameNumber;
 	/** The ar timestamp of when the LastReceivedFrame was last updated */
diff --git a/Engine/Plugins/Runtime/AppleARKit/Source/AppleARKit/Private/AppleARKitTextures.cpp b/Engine/Plugins/Runtime/AppleARKit/Source/AppleARKit/Private/AppleARKitTextures.cpp
--- a/Engine/Plugins/Runtime/AppleARKit/Source/AppleARKit/Private/AppleARKitTextures.cpp
+++ b/Engine/Plugins/Runtime/AppleARKit/Source/AppleARKit/Private/AppleARKitTextures.cpp
@@ -1,6 +1,7 @@
 // Copyright 1998-2018 Epic Games, Inc. All Rights Reserved.
 
 #include "AppleARKitTextures.h"
+#include "AppleARKitModule.h"
 #include "ExternalTexture.h"
 #include "RenderingThread.h"
 #include "Containers/DynamicRHIResourceArray.h"
@@ -55,6 +56,7 @@ void UAppleARKitTextureCameraImage::Init(float InTimestamp, CVPixelBufferRef InC
 		Size.X = CVPixelBufferGetWidth(CameraImage);
 		Size.Y = CVPixelBufferGetHeight(CameraImage);
 	}
+	//@todo joeg - Update the render resources
 }
 
 #endif
