Fixed a threading bug with the async gpu compression on iOS
Added support for a synchronous path to augment the async path

diff --git a/Engine/Plugins/Experimental/AppleImageUtils/Source/AppleImageUtils/Private/AppleImageUtilsPlugin.cpp b/Engine/Plugins/Experimental/AppleImageUtils/Source/AppleImageUtils/Private/AppleImageUtilsPlugin.cpp
--- a/Engine/Plugins/Experimental/AppleImageUtils/Source/AppleImageUtils/Private/AppleImageUtilsPlugin.cpp
+++ b/Engine/Plugins/Experimental/AppleImageUtils/Source/AppleImageUtils/Private/AppleImageUtilsPlugin.cpp
@@ -33,6 +33,12 @@ class FAppleImageUtilsPlugin :
 	virtual TSharedPtr<FAppleImageUtilsConversionTaskBase, ESPMode::ThreadSafe> ConvertToTIFF(UTexture* SourceImage, bool bWantColor = true, bool bUseGpu = true, float Scale = 1.f, ETextureRotationDirection Rotate = ETextureRotationDirection::None) override;
 #if SUPPORTS_IMAGE_UTILS_1_0
 	virtual CGImageRef UTexture2DToCGImage(UTexture2D* Source) override;
+	virtual void ConvertToJPEG(CIImage* SourceImage, TArray<uint8>& OutBytes, int32 Quality = 85, bool bWantColor = true, bool bUseGpu = true, float Scale = 1.f, ETextureRotationDirection Rotate = ETextureRotationDirection::None) override;
+#if SUPPORTS_IMAGE_UTILS_2_1
+	virtual void ConvertToHEIF(CIImage* SourceImage, TArray<uint8>& OutBytes, int32 Quality = 85,  bool bWantColor = true, bool bUseGpu = true, float Scale = 1.f, ETextureRotationDirection Rotate = ETextureRotationDirection::None) override;
+	virtual void ConvertToPNG(CIImage* SourceImage, TArray<uint8>& OutBytes, bool bWantColor = true, bool bUseGpu = true, float Scale = 1.f, ETextureRotationDirection Rotate = ETextureRotationDirection::None) override;
+	virtual void ConvertToTIFF(CIImage* SourceImage, TArray<uint8>& OutBytes, bool bWantColor = true, bool bUseGpu = true, float Scale = 1.f, ETextureRotationDirection Rotate = ETextureRotationDirection::None) override;
+#endif
 #endif
 };
 
@@ -62,15 +68,30 @@ public:
 };
 
 class FAppleImageUtilsConversionTask :
-	public FAppleImageUtilsConversionTaskBase,
-	public FGCObject
+	public FAppleImageUtilsConversionTaskBase
 {
 public:
-	FAppleImageUtilsConversionTask(UTexture* InSourceImage) :
+#if SUPPORTS_IMAGE_UTILS_1_0
+	FAppleImageUtilsConversionTask(CIImage* InSourceImage) :
 		FAppleImageUtilsConversionTaskBase()
 		, SourceImage(InSourceImage)
 	{
+		check(InSourceImage != nullptr);
+	}
+
+	virtual ~FAppleImageUtilsConversionTask()
+	{
+		if (SourceImage != nullptr)
+		{
+			[SourceImage release];
+		}
 	}
+#else
+	FAppleImageUtilsConversionTask() :
+		FAppleImageUtilsConversionTaskBase()
+	{
+	}
+#endif
 
 	//~ IAppleImageUtilsConversionTask
 	virtual TArray<uint8> GetData() override { return MoveTemp(ConvertedBytes); }
@@ -78,21 +99,14 @@ public:
 
 	void MarkComplete() { bIsDone = true; }
 
-	//~ FGCObject
-	virtual void AddReferencedObjects(FReferenceCollector& Collector) override;
-	//~ FGCObject
-
-	/** The UTexture that wraps the Apple image data underneath */
-	UTexture* SourceImage;
+#if SUPPORTS_IMAGE_UTILS_1_0
+	/** Must be released */
+	CIImage* SourceImage;
+#endif
 	/** Where the data is placed when the task is done */
 	TArray<uint8> ConvertedBytes;
 };
 
-void FAppleImageUtilsConversionTask::AddReferencedObjects(FReferenceCollector& Collector)
-{
-	Collector.AddReferencedObject(SourceImage);
-}
-
 #if SUPPORTS_IMAGE_UTILS_1_0
 static inline NSDictionary* ToQualityDictionary(int32 Quality)
 {
@@ -116,8 +130,11 @@ static inline CGColorSpaceRef ToColorSpace(bool bWantColor)
 	return bWantColor ? CGColorSpaceCreateWithName(kCGColorSpaceSRGB) : CGColorSpaceCreateWithName(kCGColorSpaceGenericGrayGamma2_2);
 }
 
-static inline CIImage* ToImage(IAppleImageInterface* AppleImageInterface, float Scale, ETextureRotationDirection Rotate)
+/** Note this image object must be released */
+static inline CIImage* AllocateImage(IAppleImageInterface* AppleImageInterface)
 {
+	// Touching UObjects so must be in game thread, hence not part of the autorelease pool due to multithreading
+	check(IsInGameThread());
 	check(AppleImageInterface != nullptr);
 
 	CIImage* Image = nullptr;
@@ -126,6 +143,7 @@ static inline CIImage* ToImage(IAppleImageInterface* AppleImageInterface, float
 		case EAppleTextureType::Image:
 		{
 			Image = AppleImageInterface->GetImage();
+			[Image retain];
 			break;
 		}
 
@@ -134,7 +152,7 @@ static inline CIImage* ToImage(IAppleImageInterface* AppleImageInterface, float
 			CVPixelBufferRef PixelBuffer = AppleImageInterface->GetPixelBuffer();
 			if (PixelBuffer != nullptr)
 			{
-				Image = [CIImage imageWithCVPixelBuffer: PixelBuffer];
+				Image = [[CIImage alloc] initWithCVPixelBuffer: PixelBuffer];
 			}
 			break;
 		}
@@ -144,11 +162,18 @@ static inline CIImage* ToImage(IAppleImageInterface* AppleImageInterface, float
 			IOSurfaceRef Surface = AppleImageInterface->GetSurface();
 			if (Surface != nullptr)
 			{
-				Image = [CIImage imageWithIOSurface: Surface];
+				Image = [[CIImage alloc] initWithIOSurface: Surface];
 			}
 			break;
 		}
 	}
+	return Image;
+}
+
+/** Note must be called from the processing thread since this uses the autorelease pool and assumes scoped release pools */
+static inline CIImage* ApplyScaleAndRotation(CIImage* SourceImage, float Scale, ETextureRotationDirection Rotate)
+{
+	CIImage* Image = SourceImage;
 	// Handle scaling if requested
 	if (Scale != 1.f)
 	{
@@ -204,42 +229,22 @@ TSharedPtr<FAppleImageUtilsConversionTaskBase, ESPMode::ThreadSafe> FAppleImageU
 		return MakeShared<FAppleImageUtilsFailedConversionTask, ESPMode::ThreadSafe>(FString::Printf(TEXT("ConvertToJPEG texture type (%d) was not supported"), (int32)AppleImage->GetTextureType()));
 	}
 
-	TSharedPtr<FAppleImageUtilsConversionTask, ESPMode::ThreadSafe> ConversionTask = MakeShared<FAppleImageUtilsConversionTask, ESPMode::ThreadSafe>(SourceImage);
 #if SUPPORTS_IMAGE_UTILS_1_0
 	if (!FAppleImageUtilsAvailability::SupportsImageUtils10())
 	{
 		return MakeShared<FAppleImageUtilsFailedConversionTask, ESPMode::ThreadSafe>(TEXT("ConvertToJPEG requires iOS 10.0+ or macOS 10.12+"));
 	}
 
-	AsyncTask(ENamedThreads::AnyBackgroundThreadNormalTask, [ConversionTask, Quality, bWantColor, bUseGpu, Scale, Rotate]()
+	TSharedPtr<FAppleImageUtilsConversionTask, ESPMode::ThreadSafe> ConversionTask = MakeShared<FAppleImageUtilsConversionTask, ESPMode::ThreadSafe>(AllocateImage(AppleImage));
+	AsyncTask(ENamedThreads::AnyBackgroundThreadNormalTask, [this, ConversionTask, Quality, bWantColor, bUseGpu, Scale, Rotate]()
 	{
-		SCOPED_AUTORELEASE_POOL;
-
-		IAppleImageInterface* AppleImageInterface = Cast<IAppleImageInterface>(ConversionTask->SourceImage);
-
-		// Convert to Apple objects
-		CIContext* ConversionContext = [CIContext contextWithOptions: ToCpuDictionary(bUseGpu)];
-		CGColorSpaceRef ColorSpace = ToColorSpace(bWantColor);
-		CIImage* Image = ToImage(AppleImageInterface, Scale, Rotate);
-
-		// This will perform the work on the GPU or inline on this thread
-		NSData* ConvertedData = [ConversionContext JPEGRepresentationOfImage: Image colorSpace: ColorSpace options: ToQualityDictionary(Quality)];
-		if (ConvertedData != nullptr)
-		{
-			uint32 CompressedSize = ConvertedData.length;
-			ConversionTask->ConvertedBytes.AddUninitialized(CompressedSize);
-			FPlatformMemory::Memcpy(ConversionTask->ConvertedBytes.GetData(), [ConvertedData bytes], CompressedSize);
-		}
-
-		CGColorSpaceRelease(ColorSpace);
-		
+		ConvertToJPEG(ConversionTask->SourceImage, ConversionTask->ConvertedBytes, Quality, bWantColor, bUseGpu, Scale, Rotate);
 		// Notify any async listeners that we are done
 		ConversionTask->MarkComplete();
 	});
-#else
-	return MakeShared<FAppleImageUtilsFailedConversionTask, ESPMode::ThreadSafe>(TEXT("ConvertToJPEG requires iOS 10.0+ or macOS 10.12+"));
-#endif
 	return ConversionTask;
+#endif
+	return MakeShared<FAppleImageUtilsFailedConversionTask, ESPMode::ThreadSafe>(TEXT("ConvertToJPEG requires iOS 10.0+ or macOS 10.12+"));
 }
 
 TSharedPtr<FAppleImageUtilsConversionTaskBase, ESPMode::ThreadSafe> FAppleImageUtilsPlugin::ConvertToHEIF(UTexture* SourceImage, int32 Quality, bool bWantColor, bool bUseGpu, float Scale, ETextureRotationDirection Rotate)
@@ -255,42 +260,22 @@ TSharedPtr<FAppleImageUtilsConversionTaskBase, ESPMode::ThreadSafe> FAppleImageU
 		return MakeShared<FAppleImageUtilsFailedConversionTask, ESPMode::ThreadSafe>(FString::Printf(TEXT("ConvertToHEIF texture type (%d) was not supported"), (int32)AppleImage->GetTextureType()));
 	}
 
-	TSharedPtr<FAppleImageUtilsConversionTask, ESPMode::ThreadSafe> ConversionTask = MakeShared<FAppleImageUtilsConversionTask, ESPMode::ThreadSafe>(SourceImage);
 #if SUPPORTS_IMAGE_UTILS_2_1
 	if (!FAppleImageUtilsAvailability::SupportsImageUtils21())
 	{
 		return MakeShared<FAppleImageUtilsFailedConversionTask, ESPMode::ThreadSafe>(TEXT("ConvertToHEIF requires iOS 11.0+ or macOS 10.13.4+"));
 	}
 
-	AsyncTask(ENamedThreads::AnyBackgroundThreadNormalTask, [ConversionTask, Quality, bWantColor, bUseGpu, Scale, Rotate]()
+	TSharedPtr<FAppleImageUtilsConversionTask, ESPMode::ThreadSafe> ConversionTask = MakeShared<FAppleImageUtilsConversionTask, ESPMode::ThreadSafe>(AllocateImage(AppleImage));
+	AsyncTask(ENamedThreads::AnyBackgroundThreadNormalTask, [this, ConversionTask, Quality, bWantColor, bUseGpu, Scale, Rotate]()
 	{
-		SCOPED_AUTORELEASE_POOL;
-
-		IAppleImageInterface* AppleImageInterface = Cast<IAppleImageInterface>(ConversionTask->SourceImage);
-
-		// Convert to Apple objects
-		CIContext* ConversionContext = [CIContext contextWithOptions: ToCpuDictionary(bUseGpu)];
-		CGColorSpaceRef ColorSpace = ToColorSpace(bWantColor);
-		CIImage* Image = ToImage(AppleImageInterface, Scale, Rotate);
-
-		// This will perform the work on the GPU or inline on this thread
-		NSData* ConvertedData = [ConversionContext HEIFRepresentationOfImage: Image format: kCIFormatARGB8 colorSpace: ColorSpace options: ToQualityDictionary(Quality)];
-		if (ConvertedData != nullptr)
-		{
-			uint32 CompressedSize = ConvertedData.length;
-			ConversionTask->ConvertedBytes.AddUninitialized(CompressedSize);
-			FPlatformMemory::Memcpy(ConversionTask->ConvertedBytes.GetData(), [ConvertedData bytes], CompressedSize);
-		}
-
-		CGColorSpaceRelease(ColorSpace);
-
+		ConvertToHEIF(ConversionTask->SourceImage, ConversionTask->ConvertedBytes, Quality, bWantColor, bUseGpu, Scale, Rotate);
 		// Notify any async listeners that we are done
 		ConversionTask->MarkComplete();
 	});
-#else
-	return MakeShared<FAppleImageUtilsFailedConversionTask, ESPMode::ThreadSafe>(TEXT("ConvertToHEIF requires iOS 11.0+ or macOS 10.13.4+"));
-#endif
 	return ConversionTask;
+#endif
+	return MakeShared<FAppleImageUtilsFailedConversionTask, ESPMode::ThreadSafe>(TEXT("ConvertToHEIF requires iOS 11.0+ or macOS 10.13.4+"));
 }
 
 TSharedPtr<FAppleImageUtilsConversionTaskBase, ESPMode::ThreadSafe> FAppleImageUtilsPlugin::ConvertToPNG(UTexture* SourceImage, bool bWantColor, bool bUseGpu, float Scale, ETextureRotationDirection Rotate)
@@ -306,42 +291,22 @@ TSharedPtr<FAppleImageUtilsConversionTaskBase, ESPMode::ThreadSafe> FAppleImageU
 		return MakeShared<FAppleImageUtilsFailedConversionTask, ESPMode::ThreadSafe>(FString::Printf(TEXT("ConvertToPNG texture type (%d) was not supported"), (int32)AppleImage->GetTextureType()));
 	}
 
-	TSharedPtr<FAppleImageUtilsConversionTask, ESPMode::ThreadSafe> ConversionTask = MakeShared<FAppleImageUtilsConversionTask, ESPMode::ThreadSafe>(SourceImage);
-#if SUPPORTS_IMAGE_UTILS_2_0
-	if (!FAppleImageUtilsAvailability::SupportsImageUtils20())
+#if SUPPORTS_IMAGE_UTILS_2_1
+	if (!FAppleImageUtilsAvailability::SupportsImageUtils21())
 	{
-		return MakeShared<FAppleImageUtilsFailedConversionTask, ESPMode::ThreadSafe>(TEXT("ConvertToPNG requires iOS 11.0+ or macOS 10.13+"));
+		return MakeShared<FAppleImageUtilsFailedConversionTask, ESPMode::ThreadSafe>(TEXT("ConvertToPNG requires iOS 11.0+ or macOS 10.13.4+"));
 	}
 
-	AsyncTask(ENamedThreads::AnyBackgroundThreadNormalTask, [ConversionTask, bWantColor, bUseGpu, Scale, Rotate]()
+	TSharedPtr<FAppleImageUtilsConversionTask, ESPMode::ThreadSafe> ConversionTask = MakeShared<FAppleImageUtilsConversionTask, ESPMode::ThreadSafe>(AllocateImage(AppleImage));
+	AsyncTask(ENamedThreads::AnyBackgroundThreadNormalTask, [this, ConversionTask, bWantColor, bUseGpu, Scale, Rotate]()
 	{
-		SCOPED_AUTORELEASE_POOL;
-
-		IAppleImageInterface* AppleImageInterface = Cast<IAppleImageInterface>(ConversionTask->SourceImage);
-
-		// Convert to Apple objects
-		CIContext* ConversionContext = [CIContext contextWithOptions: ToCpuDictionary(bUseGpu)];
-		CGColorSpaceRef ColorSpace = ToColorSpace(bWantColor);
-		CIImage* Image = ToImage(AppleImageInterface, Scale, Rotate);
-
-		// This will perform the work on the GPU or inline on this thread
-		NSData* ConvertedData = [ConversionContext PNGRepresentationOfImage: Image format: kCIFormatARGB8 colorSpace: ColorSpace options: @{}];
-		if (ConvertedData != nullptr)
-		{
-			uint32 CompressedSize = ConvertedData.length;
-			ConversionTask->ConvertedBytes.AddUninitialized(CompressedSize);
-			FPlatformMemory::Memcpy(ConversionTask->ConvertedBytes.GetData(), [ConvertedData bytes], CompressedSize);
-		}
-
-		CGColorSpaceRelease(ColorSpace);
-
+		ConvertToPNG(ConversionTask->SourceImage, ConversionTask->ConvertedBytes, bWantColor, bUseGpu, Scale, Rotate);
 		// Notify any async listeners that we are done
 		ConversionTask->MarkComplete();
 	});
-#else
-	return MakeShared<FAppleImageUtilsFailedConversionTask, ESPMode::ThreadSafe>(TEXT("ConvertToPNG requires iOS 11.0+ or macOS 10.13+"));
-#endif
 	return ConversionTask;
+#endif
+	return MakeShared<FAppleImageUtilsFailedConversionTask, ESPMode::ThreadSafe>(TEXT("ConvertToPNG requires iOS 11.0+ or macOS 10.13.4+"));
 }
 
 TSharedPtr<FAppleImageUtilsConversionTaskBase, ESPMode::ThreadSafe> FAppleImageUtilsPlugin::ConvertToTIFF(UTexture* SourceImage, bool bWantColor, bool bUseGpu, float Scale, ETextureRotationDirection Rotate)
@@ -357,42 +322,22 @@ TSharedPtr<FAppleImageUtilsConversionTaskBase, ESPMode::ThreadSafe> FAppleImageU
 		return MakeShared<FAppleImageUtilsFailedConversionTask, ESPMode::ThreadSafe>(FString::Printf(TEXT("ConvertToTIFF texture type (%d) was not supported"), (int32)AppleImage->GetTextureType()));
 	}
 
-	TSharedPtr<FAppleImageUtilsConversionTask, ESPMode::ThreadSafe> ConversionTask = MakeShared<FAppleImageUtilsConversionTask, ESPMode::ThreadSafe>(SourceImage);
-#if SUPPORTS_IMAGE_UTILS_1_0
-	if (!FAppleImageUtilsAvailability::SupportsImageUtils10())
+#if SUPPORTS_IMAGE_UTILS_2_1
+	if (!FAppleImageUtilsAvailability::SupportsImageUtils21())
 	{
-		return MakeShared<FAppleImageUtilsFailedConversionTask, ESPMode::ThreadSafe>(TEXT("ConvertToTIFF requires iOS 10.0+ or macOS 10.12+"));
+		return MakeShared<FAppleImageUtilsFailedConversionTask, ESPMode::ThreadSafe>(TEXT("ConvertToTIFF requires iOS 11.0+ or macOS 10.13.4+"));
 	}
 
-	AsyncTask(ENamedThreads::AnyBackgroundThreadNormalTask, [ConversionTask, bWantColor, bUseGpu, Scale, Rotate]()
+	TSharedPtr<FAppleImageUtilsConversionTask, ESPMode::ThreadSafe> ConversionTask = MakeShared<FAppleImageUtilsConversionTask, ESPMode::ThreadSafe>(AllocateImage(AppleImage));
+	AsyncTask(ENamedThreads::AnyBackgroundThreadNormalTask, [this, ConversionTask, bWantColor, bUseGpu, Scale, Rotate]()
 	{
-		SCOPED_AUTORELEASE_POOL;
-
-		IAppleImageInterface* AppleImageInterface = Cast<IAppleImageInterface>(ConversionTask->SourceImage);
-
-		// Convert to Apple objects
-		CIContext* ConversionContext = [CIContext contextWithOptions: ToCpuDictionary(bUseGpu)];
-		CGColorSpaceRef ColorSpace = ToColorSpace(bWantColor);
-		CIImage* Image = ToImage(AppleImageInterface, Scale, Rotate);
-
-		// This will perform the work on the GPU or inline on this thread
-		NSData* ConvertedData = [ConversionContext TIFFRepresentationOfImage: Image format: kCIFormatARGB8 colorSpace: ColorSpace options: @{}];
-		if (ConvertedData != nullptr)
-		{
-			uint32 CompressedSize = ConvertedData.length;
-			ConversionTask->ConvertedBytes.AddUninitialized(CompressedSize);
-			FPlatformMemory::Memcpy(ConversionTask->ConvertedBytes.GetData(), [ConvertedData bytes], CompressedSize);
-		}
-
-		CGColorSpaceRelease(ColorSpace);
-
+		ConvertToTIFF(ConversionTask->SourceImage, ConversionTask->ConvertedBytes, bWantColor, bUseGpu, Scale, Rotate);
 		// Notify any async listeners that we are done
 		ConversionTask->MarkComplete();
 	});
-#else
-	return MakeShared<FAppleImageUtilsFailedConversionTask, ESPMode::ThreadSafe>(TEXT("ConvertToTIFF requires iOS 10.0+ or macOS 10.12+"));
-#endif
 	return ConversionTask;
+#endif
+	return MakeShared<FAppleImageUtilsFailedConversionTask, ESPMode::ThreadSafe>(TEXT("ConvertToTIFF requires iOS 11.0+ or macOS 10.13.4+"));
 }
 
 #if SUPPORTS_IMAGE_UTILS_1_0
@@ -458,6 +403,92 @@ CGImageRef FAppleImageUtilsPlugin::UTexture2DToCGImage(UTexture2D* Source)
 	return ImageRef;
 }
 
+void FAppleImageUtilsPlugin::ConvertToJPEG(CIImage* SourceImage, TArray<uint8>& OutBytes, int32 Quality, bool bWantColor, bool bUseGpu, float Scale, ETextureRotationDirection Rotate)
+{
+	SCOPED_AUTORELEASE_POOL;
+
+	// Convert to Apple objects
+	CIContext* ConversionContext = [CIContext contextWithOptions: ToCpuDictionary(bUseGpu)];
+	CGColorSpaceRef ColorSpace = ToColorSpace(bWantColor);
+	CIImage* Image = ApplyScaleAndRotation(SourceImage, Scale, Rotate);
+
+	// This will perform the work on the GPU or inline on this thread
+	NSData* ConvertedData = [ConversionContext JPEGRepresentationOfImage: Image colorSpace: ColorSpace options: ToQualityDictionary(Quality)];
+	if (ConvertedData != nullptr)
+	{
+		uint32 CompressedSize = ConvertedData.length;
+		OutBytes.AddUninitialized(CompressedSize);
+		FPlatformMemory::Memcpy(OutBytes.GetData(), [ConvertedData bytes], CompressedSize);
+	}
+
+	CGColorSpaceRelease(ColorSpace);
+}
+
+#if SUPPORTS_IMAGE_UTILS_2_1
+void FAppleImageUtilsPlugin::ConvertToHEIF(CIImage* SourceImage, TArray<uint8>& OutBytes, int32 Quality,  bool bWantColor, bool bUseGpu, float Scale, ETextureRotationDirection Rotate)
+{
+	SCOPED_AUTORELEASE_POOL;
+
+	// Convert to Apple objects
+	CIContext* ConversionContext = [CIContext contextWithOptions: ToCpuDictionary(bUseGpu)];
+	CGColorSpaceRef ColorSpace = ToColorSpace(bWantColor);
+	CIImage* Image = ApplyScaleAndRotation(SourceImage, Scale, Rotate);
+
+	// This will perform the work on the GPU or inline on this thread
+	NSData* ConvertedData = [ConversionContext HEIFRepresentationOfImage: Image format: kCIFormatARGB8 colorSpace: ColorSpace options: ToQualityDictionary(Quality)];
+	if (ConvertedData != nullptr)
+	{
+		uint32 CompressedSize = ConvertedData.length;
+		OutBytes.AddUninitialized(CompressedSize);
+		FPlatformMemory::Memcpy(OutBytes.GetData(), [ConvertedData bytes], CompressedSize);
+	}
+
+	CGColorSpaceRelease(ColorSpace);
+}
+
+void FAppleImageUtilsPlugin::ConvertToPNG(CIImage* SourceImage, TArray<uint8>& OutBytes, bool bWantColor, bool bUseGpu, float Scale, ETextureRotationDirection Rotate)
+{
+	SCOPED_AUTORELEASE_POOL;
+
+	// Convert to Apple objects
+	CIContext* ConversionContext = [CIContext contextWithOptions: ToCpuDictionary(bUseGpu)];
+	CGColorSpaceRef ColorSpace = ToColorSpace(bWantColor);
+	CIImage* Image = ApplyScaleAndRotation(SourceImage, Scale, Rotate);
+
+	// This will perform the work on the GPU or inline on this thread
+	NSData* ConvertedData = [ConversionContext PNGRepresentationOfImage: Image format: kCIFormatARGB8 colorSpace: ColorSpace options: @{}];
+	if (ConvertedData != nullptr)
+	{
+		uint32 CompressedSize = ConvertedData.length;
+		OutBytes.AddUninitialized(CompressedSize);
+		FPlatformMemory::Memcpy(OutBytes.GetData(), [ConvertedData bytes], CompressedSize);
+	}
+
+	CGColorSpaceRelease(ColorSpace);
+}
+
+void FAppleImageUtilsPlugin::ConvertToTIFF(CIImage* SourceImage, TArray<uint8>& OutBytes, bool bWantColor, bool bUseGpu, float Scale, ETextureRotationDirection Rotate)
+{
+	SCOPED_AUTORELEASE_POOL;
+
+	// Convert to Apple objects
+	CIContext* ConversionContext = [CIContext contextWithOptions: ToCpuDictionary(bUseGpu)];
+	CGColorSpaceRef ColorSpace = ToColorSpace(bWantColor);
+	CIImage* Image = ApplyScaleAndRotation(SourceImage, Scale, Rotate);
+
+	// This will perform the work on the GPU or inline on this thread
+	NSData* ConvertedData = [ConversionContext TIFFRepresentationOfImage: Image format: kCIFormatARGB8 colorSpace: ColorSpace options: @{}];
+	if (ConvertedData != nullptr)
+	{
+		uint32 CompressedSize = ConvertedData.length;
+		OutBytes.AddUninitialized(CompressedSize);
+		FPlatformMemory::Memcpy(OutBytes.GetData(), [ConvertedData bytes], CompressedSize);
+	}
+
+	CGColorSpaceRelease(ColorSpace);
+}
+#endif
+
 	#pragma clang diagnostic pop
 #endif
 
diff --git a/Engine/Plugins/Experimental/AppleImageUtils/Source/AppleImageUtils/Public/IAppleImageUtilsPlugin.h b/Engine/Plugins/Experimental/AppleImageUtils/Source/AppleImageUtils/Public/IAppleImageUtilsPlugin.h
--- a/Engine/Plugins/Experimental/AppleImageUtils/Source/AppleImageUtils/Public/IAppleImageUtilsPlugin.h
+++ b/Engine/Plugins/Experimental/AppleImageUtils/Source/AppleImageUtils/Public/IAppleImageUtilsPlugin.h
@@ -144,6 +144,59 @@ public:
 	 * @return the CGImage if conversion was possible
 	 */
 	virtual CGImageRef UTexture2DToCGImage(UTexture2D* Source) = 0;
+
+	/**
+	 * Converts a image to an array of JPEG data synchronously
+	 *
+	 * @param SourceImage the image to compress
+	 * @param OutBytes the buffer that is populated during compression
+	 * @param Quality the quality level to compress to
+	 * @param bWantColor whether the JPEG is color (true) or monochrome (false)
+	 * @param bUseGpu whether to use the GPU (true) or the CPU (false) to compress
+	 * @param Scale whether to scale the image before converting, defaults to no scale operation
+	 * @param Rotate a direction to rotate the image in during conversion, defaults to none
+	 */
+	virtual void ConvertToJPEG(CIImage* SourceImage, TArray<uint8>& OutBytes, int32 Quality = 85, bool bWantColor = true, bool bUseGpu = true, float Scale = 1.f, ETextureRotationDirection Rotate = ETextureRotationDirection::None) = 0;
+
+#if SUPPORTS_IMAGE_UTILS_2_1
+	/**
+	 * Converts a image to an array of HEIF data synchronously
+	 *
+	 * @param SourceImage the image to compress
+	 * @param OutBytes the buffer that is populated during compression
+	 * @param Quality the quality level to compress to
+	 * @param bWantColor whether the HEIF is color (true) or monochrome (false)
+	 * @param bUseGpu whether to use the GPU (true) or the CPU (false) to compress
+	 * @param Scale whether to scale the image before converting, defaults to no scale operation
+	 * @param Rotate a direction to rotate the image in during conversion, defaults to none
+	 */
+	virtual void ConvertToHEIF(CIImage* SourceImage, TArray<uint8>& OutBytes, int32 Quality = 85,  bool bWantColor = true, bool bUseGpu = true, float Scale = 1.f, ETextureRotationDirection Rotate = ETextureRotationDirection::None) = 0;
+
+	/**
+	 * Converts a image to an array of PNG data synchronously
+	 *
+	 * @param SourceImage the image to compress
+	 * @param OutBytes the buffer that is populated during compression
+	 * @param Quality the quality level to compress to
+	 * @param bWantColor whether the PNG is color (true) or monochrome (false)
+	 * @param bUseGpu whether to use the GPU (true) or the CPU (false) to compress
+	 * @param Scale whether to scale the image before converting, defaults to no scale operation
+	 * @param Rotate a direction to rotate the image in during conversion, defaults to none
+	 */
+	virtual void ConvertToPNG(CIImage* SourceImage, TArray<uint8>& OutBytes, bool bWantColor = true, bool bUseGpu = true, float Scale = 1.f, ETextureRotationDirection Rotate = ETextureRotationDirection::None) = 0;
+
+	/**
+	 * Converts a image to an array of TIFF data synchronously
+	 *
+	 * @param SourceImage the image to compress
+	 * @param OutBytes the buffer that is populated during compression
+	 * @param Quality the quality level to compress to
+	 * @param bUseGpu whether to use the GPU (true) or the CPU (false) to compress
+	 * @param Scale whether to scale the image before converting, defaults to no scale operation
+	 * @param Rotate a direction to rotate the image in during conversion, defaults to none
+	 */
+	virtual void ConvertToTIFF(CIImage* SourceImage, TArray<uint8>& OutBytes, bool bWantColor = true, bool bUseGpu = true, float Scale = 1.f, ETextureRotationDirection Rotate = ETextureRotationDirection::None) = 0;
+#endif
 #endif
 };
 
