Moved where the texture is released to be at conversion completion instead of task destruction
Added tracking of the time it takes to perform texture compression

diff --git a/Engine/Plugins/Experimental/AppleImageUtils/Source/AppleImageUtils/Private/AppleImageUtilsPlugin.cpp b/Engine/Plugins/Experimental/AppleImageUtils/Source/AppleImageUtils/Private/AppleImageUtilsPlugin.cpp
--- a/Engine/Plugins/Experimental/AppleImageUtils/Source/AppleImageUtils/Private/AppleImageUtilsPlugin.cpp
+++ b/Engine/Plugins/Experimental/AppleImageUtils/Source/AppleImageUtils/Private/AppleImageUtilsPlugin.cpp
@@ -77,15 +77,10 @@ public:
 		, SourceImage(InSourceImage)
 	{
 		check(InSourceImage != nullptr);
+		StartTime = FPlatformTime::Seconds();
 	}
 
-	virtual ~FAppleImageUtilsConversionTask()
-	{
-		if (SourceImage != nullptr)
-		{
-			[SourceImage release];
-		}
-	}
+	virtual ~FAppleImageUtilsConversionTask() { }
 #else
 	FAppleImageUtilsConversionTask() :
 		FAppleImageUtilsConversionTaskBase()
@@ -97,7 +92,17 @@ public:
 	virtual TArray<uint8> GetData() override { return MoveTemp(ConvertedBytes); }
 	//~ IAppleImageUtilsConversionTask
 
-	void MarkComplete() { bIsDone = true; }
+	void MarkComplete()
+	{
+		// Release on completion since we don't need it anymore
+		if (SourceImage != nullptr)
+		{
+			[SourceImage release];
+			SourceImage = nullptr;
+		}
+		EndTime = FPlatformTime::Seconds();
+		bIsDone = true;
+	}
 
 #if SUPPORTS_IMAGE_UTILS_1_0
 	/** Must be released */
diff --git a/Engine/Plugins/Experimental/AppleImageUtils/Source/AppleImageUtils/Public/IAppleImageUtilsPlugin.h b/Engine/Plugins/Experimental/AppleImageUtils/Source/AppleImageUtils/Public/IAppleImageUtilsPlugin.h
--- a/Engine/Plugins/Experimental/AppleImageUtils/Source/AppleImageUtils/Public/IAppleImageUtilsPlugin.h
+++ b/Engine/Plugins/Experimental/AppleImageUtils/Source/AppleImageUtils/Public/IAppleImageUtilsPlugin.h
@@ -30,6 +30,8 @@ public:
 	virtual FString GetErrorReason() const = 0;
 	/** @return whether the task has completed or not */
 	virtual bool IsDone() const = 0;
+	/** @return how long it took this task to complete */
+	virtual float GetElapsedTime() const = 0;
 	/** @return the data once the task has completed or empty array if still in progress. NOTE: this should use MoveTemp() to avoid duplication so call once :) */
 	virtual TArray<uint8> GetData() = 0;
 };
@@ -48,10 +50,13 @@ public:
 	virtual bool IsDone() const override { return bIsDone; }
 	virtual bool HadError() const override { return bHadError; }
 	virtual FString GetErrorReason() const override { return Error; }
+	virtual float GetElapsedTime() const override { return (EndTime - StartTime) * 1000.f; }
 
 protected:
 	FThreadSafeBool bIsDone;
 	FThreadSafeBool bHadError;
+	double StartTime = 0.f;
+	double EndTime = 0.f;
 	FString Error;
 };
 
