Fixed missing support check

diff --git a/Engine/Plugins/Experimental/AppleImageUtils/Source/AppleImageUtils/Private/AppleImageUtilsPlugin.cpp b/Engine/Plugins/Experimental/AppleImageUtils/Source/AppleImageUtils/Private/AppleImageUtilsPlugin.cpp
--- a/Engine/Plugins/Experimental/AppleImageUtils/Source/AppleImageUtils/Private/AppleImageUtilsPlugin.cpp
+++ b/Engine/Plugins/Experimental/AppleImageUtils/Source/AppleImageUtils/Private/AppleImageUtilsPlugin.cpp
@@ -94,12 +94,14 @@ public:
 
 	void MarkComplete()
 	{
+#if SUPPORTS_IMAGE_UTILS_1_0
 		// Release on completion since we don't need it anymore
 		if (SourceImage != nullptr)
 		{
 			[SourceImage release];
 			SourceImage = nullptr;
 		}
+#endif
 		EndTime = FPlatformTime::Seconds();
 		bIsDone = true;
 	}
