Moved AppleImageUtils from Experimental to Runtime

diff --git a/Engine/Plugins/Experimental/AppleImageUtils/AppleImageUtils.uplugin b/Engine/Plugins/Experimental/AppleImageUtils/AppleImageUtils.uplugin
deleted file mode 100644
diff --git a/Engine/Plugins/Experimental/AppleImageUtils/Source/AppleImageUtils/AppleImageUtils.Build.cs b/Engine/Plugins/Experimental/AppleImageUtils/Source/AppleImageUtils/AppleImageUtils.Build.cs
deleted file mode 100644
diff --git a/Engine/Plugins/Experimental/AppleImageUtils/Source/AppleImageUtils/Private/AppleImageUtilsBlueprintProxy.cpp b/Engine/Plugins/Experimental/AppleImageUtils/Source/AppleImageUtils/Private/AppleImageUtilsBlueprintProxy.cpp
deleted file mode 100644
diff --git a/Engine/Plugins/Experimental/AppleImageUtils/Source/AppleImageUtils/Private/AppleImageUtilsPlugin.cpp b/Engine/Plugins/Experimental/AppleImageUtils/Source/AppleImageUtils/Private/AppleImageUtilsPlugin.cpp
deleted file mode 100644
diff --git a/Engine/Plugins/Experimental/AppleImageUtils/Source/AppleImageUtils/Public/AppleImageUtilsAvailability.h b/Engine/Plugins/Experimental/AppleImageUtils/Source/AppleImageUtils/Public/AppleImageUtilsAvailability.h
deleted file mode 100644
diff --git a/Engine/Plugins/Experimental/AppleImageUtils/Source/AppleImageUtils/Public/AppleImageUtilsBlueprintProxy.h b/Engine/Plugins/Experimental/AppleImageUtils/Source/AppleImageUtils/Public/AppleImageUtilsBlueprintProxy.h
deleted file mode 100644
diff --git a/Engine/Plugins/Experimental/AppleImageUtils/Source/AppleImageUtils/Public/AppleImageUtilsTypes.h b/Engine/Plugins/Experimental/AppleImageUtils/Source/AppleImageUtils/Public/AppleImageUtilsTypes.h
deleted file mode 100644
diff --git a/Engine/Plugins/Experimental/AppleImageUtils/Source/AppleImageUtils/Public/IAppleImageUtilsPlugin.h b/Engine/Plugins/Experimental/AppleImageUtils/Source/AppleImageUtils/Public/IAppleImageUtilsPlugin.h
deleted file mode 100644
diff --git a/Engine/Plugins/Experimental/AppleImageUtils/Source/AppleImageUtilsBlueprintSupport/AppleImageUtilsBlueprintSupport.Build.cs b/Engine/Plugins/Experimental/AppleImageUtils/Source/AppleImageUtilsBlueprintSupport/AppleImageUtilsBlueprintSupport.Build.cs
deleted file mode 100644
diff --git a/Engine/Plugins/Experimental/AppleImageUtils/Source/AppleImageUtilsBlueprintSupport/Classes/AppleImageUtilsBlueprintSupport.h b/Engine/Plugins/Experimental/AppleImageUtils/Source/AppleImageUtilsBlueprintSupport/Classes/AppleImageUtilsBlueprintSupport.h
deleted file mode 100644
diff --git a/Engine/Plugins/Experimental/AppleImageUtils/Source/AppleImageUtilsBlueprintSupport/Private/AppleImageUtilsBlueprintSupport.cpp b/Engine/Plugins/Experimental/AppleImageUtils/Source/AppleImageUtilsBlueprintSupport/Private/AppleImageUtilsBlueprintSupport.cpp
deleted file mode 100644
diff --git a/Engine/Plugins/Runtime/AppleImageUtils/AppleImageUtils.uplugin b/Engine/Plugins/Runtime/AppleImageUtils/AppleImageUtils.uplugin
new file mode 100644
--- /dev/null
+++ b/Engine/Plugins/Runtime/AppleImageUtils/AppleImageUtils.uplugin
@@ -0,0 +1,40 @@
+{
+	"FileVersion" : 3,
+	"Version" : 1,
+	"VersionName" : "1.0",
+	"FriendlyName" : "Apple Image Utils",
+	"Description" : "Utilities that operate on CIImage, CVPixelBuffer, IOSurface, etc.",
+	"Category" : "Experimental",
+	"CreatedBy" : "Epic Games, Inc.",
+	"CreatedByURL" : "http://epicgames.com",
+	"DocsURL" : "",
+	"MarketplaceURL" : "",
+	"SupportURL" : "",
+	"EnabledByDefault" : true,
+	"CanContainContent" : false,
+	"IsBetaVersion" : true,
+	"Installed" : false,
+	"Modules" :
+	[
+		{
+			"Name" : "AppleImageUtils",
+			"Type" : "Runtime",
+			"LoadingPhase" : "Default",
+			"BlacklistPlatforms" :
+			[
+				"TVOS"
+			]
+		},
+		{
+			"Name" : "AppleImageUtilsBlueprintSupport",
+			"Type" : "Developer",
+			"LoadingPhase" : "Default",
+			"WhitelistPlatforms" :
+			[
+				"Win64",
+				"Mac",
+				"Linux"
+			]
+		}
+	]
+}
diff --git a/Engine/Plugins/Runtime/AppleImageUtils/Source/AppleImageUtils/AppleImageUtils.Build.cs b/Engine/Plugins/Runtime/AppleImageUtils/Source/AppleImageUtils/AppleImageUtils.Build.cs
new file mode 100644
--- /dev/null
+++ b/Engine/Plugins/Runtime/AppleImageUtils/Source/AppleImageUtils/AppleImageUtils.Build.cs
@@ -0,0 +1,62 @@
+// Copyright 1998-2019 Epic Games, Inc. All Rights Reserved.
+
+namespace UnrealBuildTool.Rules
+{
+	public class AppleImageUtils : ModuleRules
+	{
+		public AppleImageUtils(ReadOnlyTargetRules Target) : base(Target)
+		{
+            PublicIncludePaths.AddRange(
+                new string[] {
+                    // ... add public include paths required here ...
+                }
+                );
+
+
+            PrivateIncludePaths.AddRange(
+                new string[] {
+                    "AppleImageUtils/Public",
+                    "AppleImageUtils/Private"
+                    // ... add other private include paths required here ...
+                }
+                );
+
+
+            PublicDependencyModuleNames.AddRange(
+                new string[]
+                {
+                    "Core",
+                    "Engine",
+                    // ... add other public dependencies that you statically link with here ...
+                }
+                );
+
+    		if (Target.Platform == UnrealTargetPlatform.IOS || Target.Platform == UnrealTargetPlatform.Mac)
+    		{
+                PublicFrameworks.AddRange(
+                    new string[]
+                    {
+                        "CoreImage",
+                        "ImageIO"
+                        // ... add other public dependencies that you statically link with here ...
+                    }
+                    );
+            }
+
+			PrivateDependencyModuleNames.AddRange(
+				new string[]
+				{
+    				"CoreUObject"
+					// ... add private dependencies that you statically link with here ...
+				}
+				);
+
+			DynamicallyLoadedModuleNames.AddRange(
+				new string[]
+				{
+					// ... add any modules that your module loads dynamically here ...
+				}
+				);
+		}
+	}
+}
diff --git a/Engine/Plugins/Runtime/AppleImageUtils/Source/AppleImageUtils/Private/AppleImageUtilsBlueprintProxy.cpp b/Engine/Plugins/Runtime/AppleImageUtils/Source/AppleImageUtils/Private/AppleImageUtilsBlueprintProxy.cpp
new file mode 100644
--- /dev/null
+++ b/Engine/Plugins/Runtime/AppleImageUtils/Source/AppleImageUtils/Private/AppleImageUtilsBlueprintProxy.cpp
@@ -0,0 +1,78 @@
+// Copyright 1998-2019 Epic Games, Inc. All Rights Reserved.
+
+#include "AppleImageUtilsBlueprintProxy.h"
+#include "IAppleImageUtilsPlugin.h"
+
+UAppleImageUtilsBaseAsyncTaskBlueprintProxy::UAppleImageUtilsBaseAsyncTaskBlueprintProxy(const FObjectInitializer& ObjectInitializer)
+	: Super(ObjectInitializer)
+	, bShouldTick(true)
+{
+}
+
+void UAppleImageUtilsBaseAsyncTaskBlueprintProxy::Tick(float DeltaTime)
+{
+	if (!ConversionTask.IsValid())
+	{
+		ConversionResult.Error = TEXT("Invalid conversion task");
+		bShouldTick = false;
+		OnFailure.Broadcast(ConversionResult);
+		return;
+	}
+
+	if (ConversionTask->IsDone())
+	{
+		bShouldTick = false;
+		// Fire the right delegate
+		if (!ConversionTask->HadError())
+		{
+			ConversionResult.Error = TEXT("Success");
+			ConversionResult.ImageData = ConversionTask->GetData();
+			OnSuccess.Broadcast(ConversionResult);
+		}
+		else
+		{
+			ConversionResult.Error = ConversionTask->GetErrorReason();
+			OnFailure.Broadcast(ConversionResult);
+		}
+	}
+}
+
+UAppleImageUtilsBaseAsyncTaskBlueprintProxy* UAppleImageUtilsBaseAsyncTaskBlueprintProxy::CreateProxyObjectForConvertToJPEG(UTexture* SourceImage, int32 Quality, bool bWantColor, bool bUseGpu, float Scale, ETextureRotationDirection Rotate)
+{
+	UAppleImageUtilsBaseAsyncTaskBlueprintProxy* Proxy = NewObject<UAppleImageUtilsBaseAsyncTaskBlueprintProxy>();
+	Proxy->SetFlags(RF_StrongRefOnFrame);
+
+	Proxy->ConversionTask = IAppleImageUtilsPlugin::Get().ConvertToJPEG(SourceImage, Quality, bWantColor, bUseGpu, Scale, Rotate);
+
+	return Proxy;
+}
+
+UAppleImageUtilsBaseAsyncTaskBlueprintProxy* UAppleImageUtilsBaseAsyncTaskBlueprintProxy::CreateProxyObjectForConvertToHEIF(UTexture* SourceImage, int32 Quality, bool bWantColor, bool bUseGpu, float Scale, ETextureRotationDirection Rotate)
+{
+	UAppleImageUtilsBaseAsyncTaskBlueprintProxy* Proxy = NewObject<UAppleImageUtilsBaseAsyncTaskBlueprintProxy>();
+	Proxy->SetFlags(RF_StrongRefOnFrame);
+
+	Proxy->ConversionTask = IAppleImageUtilsPlugin::Get().ConvertToHEIF(SourceImage, Quality, bWantColor, bUseGpu, Scale, Rotate);
+
+	return Proxy;
+}
+
+UAppleImageUtilsBaseAsyncTaskBlueprintProxy* UAppleImageUtilsBaseAsyncTaskBlueprintProxy::CreateProxyObjectForConvertToTIFF(UTexture* SourceImage, bool bWantColor, bool bUseGpu, float Scale, ETextureRotationDirection Rotate)
+{
+	UAppleImageUtilsBaseAsyncTaskBlueprintProxy* Proxy = NewObject<UAppleImageUtilsBaseAsyncTaskBlueprintProxy>();
+	Proxy->SetFlags(RF_StrongRefOnFrame);
+
+	Proxy->ConversionTask = IAppleImageUtilsPlugin::Get().ConvertToTIFF(SourceImage, bWantColor, bUseGpu, Scale, Rotate);
+
+	return Proxy;
+}
+
+UAppleImageUtilsBaseAsyncTaskBlueprintProxy* UAppleImageUtilsBaseAsyncTaskBlueprintProxy::CreateProxyObjectForConvertToPNG(UTexture* SourceImage, bool bWantColor, bool bUseGpu, float Scale, ETextureRotationDirection Rotate)
+{
+	UAppleImageUtilsBaseAsyncTaskBlueprintProxy* Proxy = NewObject<UAppleImageUtilsBaseAsyncTaskBlueprintProxy>();
+	Proxy->SetFlags(RF_StrongRefOnFrame);
+
+	Proxy->ConversionTask = IAppleImageUtilsPlugin::Get().ConvertToPNG(SourceImage, bWantColor, bUseGpu, Scale, Rotate);
+
+	return Proxy;
+}
diff --git a/Engine/Plugins/Runtime/AppleImageUtils/Source/AppleImageUtils/Private/AppleImageUtilsPlugin.cpp b/Engine/Plugins/Runtime/AppleImageUtils/Source/AppleImageUtils/Private/AppleImageUtilsPlugin.cpp
new file mode 100644
--- /dev/null
+++ b/Engine/Plugins/Runtime/AppleImageUtils/Source/AppleImageUtils/Private/AppleImageUtilsPlugin.cpp
@@ -0,0 +1,506 @@
+// Copyright 1998-2019 Epic Games, Inc. All Rights Reserved.
+
+#include "CoreMinimal.h"
+#include "Modules/ModuleManager.h"
+#include "IAppleImageUtilsPlugin.h"
+#include "AppleImageUtilsTypes.h"
+#include "HAL/UnrealMemory.h"
+#include "UObject/GCObject.h"
+#include "Async/Async.h"
+#include "Engine/Texture.h"
+#include "Engine/Texture2D.h"
+
+#if SUPPORTS_IMAGE_UTILS_1_0
+	#include "Apple/ApplePlatformMisc.h"
+
+	#import <CoreImage/CIContext.h>
+
+	// For runtime checks so that clang doesn't warn on targets < our SDK version
+	#pragma clang diagnostic push
+	#pragma clang diagnostic ignored "-Wpartial-availability"
+#endif
+
+class FAppleImageUtilsPlugin :
+	public IAppleImageUtilsPlugin
+{
+	/** IModuleInterface implementation */
+	virtual void StartupModule() override;
+	virtual void ShutdownModule() override;
+
+	virtual TSharedPtr<FAppleImageUtilsConversionTaskBase, ESPMode::ThreadSafe> ConvertToJPEG(UTexture* SourceImage, int32 Quality = 85, bool bWantColor = true, bool bUseGpu = true, float Scale = 1.f, ETextureRotationDirection Rotate = ETextureRotationDirection::None) override;
+	virtual TSharedPtr<FAppleImageUtilsConversionTaskBase, ESPMode::ThreadSafe> ConvertToHEIF(UTexture* SourceImage, int32 Quality = 85, bool bWantColor = true, bool bUseGpu = true, float Scale = 1.f, ETextureRotationDirection Rotate = ETextureRotationDirection::None) override;
+	virtual TSharedPtr<FAppleImageUtilsConversionTaskBase, ESPMode::ThreadSafe> ConvertToPNG(UTexture* SourceImage, bool bWantColor = true, bool bUseGpu = true, float Scale = 1.f, ETextureRotationDirection Rotate = ETextureRotationDirection::None) override;
+	virtual TSharedPtr<FAppleImageUtilsConversionTaskBase, ESPMode::ThreadSafe> ConvertToTIFF(UTexture* SourceImage, bool bWantColor = true, bool bUseGpu = true, float Scale = 1.f, ETextureRotationDirection Rotate = ETextureRotationDirection::None) override;
+#if SUPPORTS_IMAGE_UTILS_1_0
+	virtual CGImageRef UTexture2DToCGImage(UTexture2D* Source) override;
+	virtual void ConvertToJPEG(CIImage* SourceImage, TArray<uint8>& OutBytes, int32 Quality = 85, bool bWantColor = true, bool bUseGpu = true, float Scale = 1.f, ETextureRotationDirection Rotate = ETextureRotationDirection::None) override;
+#if SUPPORTS_IMAGE_UTILS_2_1
+	virtual void ConvertToHEIF(CIImage* SourceImage, TArray<uint8>& OutBytes, int32 Quality = 85,  bool bWantColor = true, bool bUseGpu = true, float Scale = 1.f, ETextureRotationDirection Rotate = ETextureRotationDirection::None) override;
+	virtual void ConvertToPNG(CIImage* SourceImage, TArray<uint8>& OutBytes, bool bWantColor = true, bool bUseGpu = true, float Scale = 1.f, ETextureRotationDirection Rotate = ETextureRotationDirection::None) override;
+	virtual void ConvertToTIFF(CIImage* SourceImage, TArray<uint8>& OutBytes, bool bWantColor = true, bool bUseGpu = true, float Scale = 1.f, ETextureRotationDirection Rotate = ETextureRotationDirection::None) override;
+#endif
+#endif
+};
+
+IMPLEMENT_MODULE(FAppleImageUtilsPlugin, AppleImageUtils);
+
+void FAppleImageUtilsPlugin::StartupModule()
+{
+}
+
+void FAppleImageUtilsPlugin::ShutdownModule()
+{
+}
+
+class FAppleImageUtilsFailedConversionTask :
+	public FAppleImageUtilsConversionTaskBase
+{
+public:
+	FAppleImageUtilsFailedConversionTask(const FString& InError) :
+		FAppleImageUtilsConversionTaskBase()
+	{
+		Error = InError;
+		bHadError = true;
+		bIsDone = true;
+	}
+
+	virtual TArray<uint8> GetData() override { return TArray<uint8>(); }
+};
+
+class FAppleImageUtilsConversionTask :
+	public FAppleImageUtilsConversionTaskBase
+{
+public:
+#if SUPPORTS_IMAGE_UTILS_1_0
+	FAppleImageUtilsConversionTask(CIImage* InSourceImage) :
+		FAppleImageUtilsConversionTaskBase()
+		, SourceImage(InSourceImage)
+	{
+		check(InSourceImage != nullptr);
+		StartTime = FPlatformTime::Seconds();
+	}
+
+	virtual ~FAppleImageUtilsConversionTask() { }
+#else
+	FAppleImageUtilsConversionTask() :
+		FAppleImageUtilsConversionTaskBase()
+	{
+	}
+#endif
+
+	//~ IAppleImageUtilsConversionTask
+	virtual TArray<uint8> GetData() override { return MoveTemp(ConvertedBytes); }
+	//~ IAppleImageUtilsConversionTask
+
+	void MarkComplete()
+	{
+#if SUPPORTS_IMAGE_UTILS_1_0
+		// Release on completion since we don't need it anymore
+		if (SourceImage != nullptr)
+		{
+			[SourceImage release];
+			SourceImage = nullptr;
+		}
+#endif
+		EndTime = FPlatformTime::Seconds();
+		bIsDone = true;
+	}
+
+#if SUPPORTS_IMAGE_UTILS_1_0
+	/** Must be released */
+	CIImage* SourceImage;
+#endif
+	/** Where the data is placed when the task is done */
+	TArray<uint8> ConvertedBytes;
+};
+
+#if SUPPORTS_IMAGE_UTILS_1_0
+static inline NSDictionary* ToQualityDictionary(int32 Quality)
+{
+	// Our code uses ints but Apple uses floats 0..1
+	Quality = FMath::Clamp(Quality, 0, 100);
+	float QualityF = float(Quality) * 0.01f;
+
+	NSDictionary* Options = [NSDictionary dictionaryWithObjectsAndKeys:[NSNumber numberWithFloat: QualityF], kCGImageDestinationLossyCompressionQuality, nil];
+	return Options;
+}
+
+static inline NSDictionary* ToCpuDictionary(bool bUseGpu)
+{
+	BOOL UseCpu = bUseGpu ? NO : YES;
+	NSDictionary* Options = [NSDictionary dictionaryWithObjectsAndKeys:[NSNumber numberWithBool: UseCpu], kCIContextUseSoftwareRenderer, nil];
+	return Options;
+}
+
+static inline CGColorSpaceRef ToColorSpace(bool bWantColor)
+{
+	return bWantColor ? CGColorSpaceCreateWithName(kCGColorSpaceSRGB) : CGColorSpaceCreateWithName(kCGColorSpaceGenericGrayGamma2_2);
+}
+
+/** Note this image object must be released */
+static inline CIImage* AllocateImage(IAppleImageInterface* AppleImageInterface)
+{
+	// Touching UObjects so must be in game thread, hence not part of the autorelease pool due to multithreading
+	check(IsInGameThread());
+	check(AppleImageInterface != nullptr);
+
+	CIImage* Image = nullptr;
+	switch (AppleImageInterface->GetTextureType())
+	{
+		case EAppleTextureType::Image:
+		{
+			Image = AppleImageInterface->GetImage();
+			[Image retain];
+			break;
+		}
+
+		case EAppleTextureType::PixelBuffer:
+		{
+			CVPixelBufferRef PixelBuffer = AppleImageInterface->GetPixelBuffer();
+			if (PixelBuffer != nullptr)
+			{
+				Image = [[CIImage alloc] initWithCVPixelBuffer: PixelBuffer];
+			}
+			break;
+		}
+
+		case EAppleTextureType::Surface:
+		{
+			IOSurfaceRef Surface = AppleImageInterface->GetSurface();
+			if (Surface != nullptr)
+			{
+				Image = [[CIImage alloc] initWithIOSurface: Surface];
+			}
+			break;
+		}
+	}
+	return Image;
+}
+
+/** Note must be called from the processing thread since this uses the autorelease pool and assumes scoped release pools */
+static inline CIImage* ApplyScaleAndRotation(CIImage* SourceImage, float Scale, ETextureRotationDirection Rotate)
+{
+	CIImage* Image = SourceImage;
+	// Handle scaling if requested
+	if (Scale != 1.f)
+	{
+		CGRect Rect = Image.extent;
+		float AspectRatio = (float)Rect.size.width / (float)Rect.size.height;
+		CIFilter* ScaleFilter = [CIFilter filterWithName: @"CILanczosScaleTransform"];
+		[ScaleFilter setValue: Image forKey: kCIInputImageKey];
+		[ScaleFilter setValue: @(Scale) forKey: kCIInputScaleKey];
+		[ScaleFilter setValue: @(AspectRatio) forKey: kCIInputAspectRatioKey];
+		Image = ScaleFilter.outputImage;
+	}
+	// Handle rotation if requested
+	switch (Rotate)
+	{
+		case ETextureRotationDirection::Left:
+			Image = [Image imageByApplyingOrientation: kCGImagePropertyOrientationLeft];
+			break;
+		case ETextureRotationDirection::Right:
+			Image = [Image imageByApplyingOrientation: kCGImagePropertyOrientationRight];
+			break;
+		case ETextureRotationDirection::Down:
+			Image = [Image imageByApplyingOrientation: kCGImagePropertyOrientationDown];
+			break;
+	}
+	return Image;
+}
+#endif
+
+static inline bool CanBeConverted(IAppleImageInterface* AppleImageInterface)
+{
+	check(AppleImageInterface != nullptr);
+
+	switch (AppleImageInterface->GetTextureType())
+	{
+		case EAppleTextureType::Image:
+		case EAppleTextureType::PixelBuffer:
+		case EAppleTextureType::Surface:
+			return true;
+	}
+	return false;
+}
+
+TSharedPtr<FAppleImageUtilsConversionTaskBase, ESPMode::ThreadSafe> FAppleImageUtilsPlugin::ConvertToJPEG(UTexture* SourceImage, int32 Quality, bool bWantColor, bool bUseGpu, float Scale, ETextureRotationDirection Rotate)
+{
+	// Make sure our interface is supported
+	IAppleImageInterface* AppleImage = Cast<IAppleImageInterface>(SourceImage);
+	if (AppleImage == nullptr)
+	{
+		return MakeShared<FAppleImageUtilsFailedConversionTask, ESPMode::ThreadSafe>(TEXT("ConvertToJPEG only supports UAppleImageInterface derived textures"));
+	}
+	if (!CanBeConverted(AppleImage))
+	{
+		return MakeShared<FAppleImageUtilsFailedConversionTask, ESPMode::ThreadSafe>(FString::Printf(TEXT("ConvertToJPEG texture type (%d) was not supported"), (int32)AppleImage->GetTextureType()));
+	}
+
+#if SUPPORTS_IMAGE_UTILS_1_0
+	if (!FAppleImageUtilsAvailability::SupportsImageUtils10())
+	{
+		return MakeShared<FAppleImageUtilsFailedConversionTask, ESPMode::ThreadSafe>(TEXT("ConvertToJPEG requires iOS 10.0+ or macOS 10.12+"));
+	}
+
+	TSharedPtr<FAppleImageUtilsConversionTask, ESPMode::ThreadSafe> ConversionTask = MakeShared<FAppleImageUtilsConversionTask, ESPMode::ThreadSafe>(AllocateImage(AppleImage));
+	AsyncTask(ENamedThreads::AnyBackgroundThreadNormalTask, [this, ConversionTask, Quality, bWantColor, bUseGpu, Scale, Rotate]()
+	{
+		ConvertToJPEG(ConversionTask->SourceImage, ConversionTask->ConvertedBytes, Quality, bWantColor, bUseGpu, Scale, Rotate);
+		// Notify any async listeners that we are done
+		ConversionTask->MarkComplete();
+	});
+	return ConversionTask;
+#endif
+	return MakeShared<FAppleImageUtilsFailedConversionTask, ESPMode::ThreadSafe>(TEXT("ConvertToJPEG requires iOS 10.0+ or macOS 10.12+"));
+}
+
+TSharedPtr<FAppleImageUtilsConversionTaskBase, ESPMode::ThreadSafe> FAppleImageUtilsPlugin::ConvertToHEIF(UTexture* SourceImage, int32 Quality, bool bWantColor, bool bUseGpu, float Scale, ETextureRotationDirection Rotate)
+{
+	// Make sure our interface is supported
+	IAppleImageInterface* AppleImage = Cast<IAppleImageInterface>(SourceImage);
+	if (AppleImage == nullptr)
+	{
+		return MakeShared<FAppleImageUtilsFailedConversionTask, ESPMode::ThreadSafe>(TEXT("ConvertToHEIF only supports UAppleImageInterface derived textures"));
+	}
+	if (!CanBeConverted(AppleImage))
+	{
+		return MakeShared<FAppleImageUtilsFailedConversionTask, ESPMode::ThreadSafe>(FString::Printf(TEXT("ConvertToHEIF texture type (%d) was not supported"), (int32)AppleImage->GetTextureType()));
+	}
+
+#if SUPPORTS_IMAGE_UTILS_2_1
+	if (!FAppleImageUtilsAvailability::SupportsImageUtils21())
+	{
+		return MakeShared<FAppleImageUtilsFailedConversionTask, ESPMode::ThreadSafe>(TEXT("ConvertToHEIF requires iOS 11.0+ or macOS 10.13.4+"));
+	}
+
+	TSharedPtr<FAppleImageUtilsConversionTask, ESPMode::ThreadSafe> ConversionTask = MakeShared<FAppleImageUtilsConversionTask, ESPMode::ThreadSafe>(AllocateImage(AppleImage));
+	AsyncTask(ENamedThreads::AnyBackgroundThreadNormalTask, [this, ConversionTask, Quality, bWantColor, bUseGpu, Scale, Rotate]()
+	{
+		ConvertToHEIF(ConversionTask->SourceImage, ConversionTask->ConvertedBytes, Quality, bWantColor, bUseGpu, Scale, Rotate);
+		// Notify any async listeners that we are done
+		ConversionTask->MarkComplete();
+	});
+	return ConversionTask;
+#endif
+	return MakeShared<FAppleImageUtilsFailedConversionTask, ESPMode::ThreadSafe>(TEXT("ConvertToHEIF requires iOS 11.0+ or macOS 10.13.4+"));
+}
+
+TSharedPtr<FAppleImageUtilsConversionTaskBase, ESPMode::ThreadSafe> FAppleImageUtilsPlugin::ConvertToPNG(UTexture* SourceImage, bool bWantColor, bool bUseGpu, float Scale, ETextureRotationDirection Rotate)
+{
+	// Make sure our interface is supported
+	IAppleImageInterface* AppleImage = Cast<IAppleImageInterface>(SourceImage);
+	if (AppleImage == nullptr)
+	{
+		return MakeShared<FAppleImageUtilsFailedConversionTask, ESPMode::ThreadSafe>(TEXT("ConvertToPNG only supports UAppleImageInterface derived textures"));
+	}
+	if (!CanBeConverted(AppleImage))
+	{
+		return MakeShared<FAppleImageUtilsFailedConversionTask, ESPMode::ThreadSafe>(FString::Printf(TEXT("ConvertToPNG texture type (%d) was not supported"), (int32)AppleImage->GetTextureType()));
+	}
+
+#if SUPPORTS_IMAGE_UTILS_2_1
+	if (!FAppleImageUtilsAvailability::SupportsImageUtils21())
+	{
+		return MakeShared<FAppleImageUtilsFailedConversionTask, ESPMode::ThreadSafe>(TEXT("ConvertToPNG requires iOS 11.0+ or macOS 10.13.4+"));
+	}
+
+	TSharedPtr<FAppleImageUtilsConversionTask, ESPMode::ThreadSafe> ConversionTask = MakeShared<FAppleImageUtilsConversionTask, ESPMode::ThreadSafe>(AllocateImage(AppleImage));
+	AsyncTask(ENamedThreads::AnyBackgroundThreadNormalTask, [this, ConversionTask, bWantColor, bUseGpu, Scale, Rotate]()
+	{
+		ConvertToPNG(ConversionTask->SourceImage, ConversionTask->ConvertedBytes, bWantColor, bUseGpu, Scale, Rotate);
+		// Notify any async listeners that we are done
+		ConversionTask->MarkComplete();
+	});
+	return ConversionTask;
+#endif
+	return MakeShared<FAppleImageUtilsFailedConversionTask, ESPMode::ThreadSafe>(TEXT("ConvertToPNG requires iOS 11.0+ or macOS 10.13.4+"));
+}
+
+TSharedPtr<FAppleImageUtilsConversionTaskBase, ESPMode::ThreadSafe> FAppleImageUtilsPlugin::ConvertToTIFF(UTexture* SourceImage, bool bWantColor, bool bUseGpu, float Scale, ETextureRotationDirection Rotate)
+{
+	// Make sure our interface is supported
+	IAppleImageInterface* AppleImage = Cast<IAppleImageInterface>(SourceImage);
+	if (AppleImage == nullptr)
+	{
+		return MakeShared<FAppleImageUtilsFailedConversionTask, ESPMode::ThreadSafe>(TEXT("ConvertToTIFF only supports UAppleImageInterface derived textures"));
+	}
+	if (!CanBeConverted(AppleImage))
+	{
+		return MakeShared<FAppleImageUtilsFailedConversionTask, ESPMode::ThreadSafe>(FString::Printf(TEXT("ConvertToTIFF texture type (%d) was not supported"), (int32)AppleImage->GetTextureType()));
+	}
+
+#if SUPPORTS_IMAGE_UTILS_2_1
+	if (!FAppleImageUtilsAvailability::SupportsImageUtils21())
+	{
+		return MakeShared<FAppleImageUtilsFailedConversionTask, ESPMode::ThreadSafe>(TEXT("ConvertToTIFF requires iOS 11.0+ or macOS 10.13.4+"));
+	}
+
+	TSharedPtr<FAppleImageUtilsConversionTask, ESPMode::ThreadSafe> ConversionTask = MakeShared<FAppleImageUtilsConversionTask, ESPMode::ThreadSafe>(AllocateImage(AppleImage));
+	AsyncTask(ENamedThreads::AnyBackgroundThreadNormalTask, [this, ConversionTask, bWantColor, bUseGpu, Scale, Rotate]()
+	{
+		ConvertToTIFF(ConversionTask->SourceImage, ConversionTask->ConvertedBytes, bWantColor, bUseGpu, Scale, Rotate);
+		// Notify any async listeners that we are done
+		ConversionTask->MarkComplete();
+	});
+	return ConversionTask;
+#endif
+	return MakeShared<FAppleImageUtilsFailedConversionTask, ESPMode::ThreadSafe>(TEXT("ConvertToTIFF requires iOS 11.0+ or macOS 10.13.4+"));
+}
+
+#if SUPPORTS_IMAGE_UTILS_1_0
+
+CGImageRef FAppleImageUtilsPlugin::UTexture2DToCGImage(UTexture2D* Source)
+{
+	check(Source != nullptr);
+
+	const EPixelFormat SourceFormat = Source->GetPixelFormat();
+	if (!(SourceFormat == PF_A8R8G8B8 || SourceFormat == PF_R8G8B8A8 || SourceFormat == PF_B8G8R8A8))
+	{
+		UE_LOG(LogTemp, Warning, TEXT("TextureToCGImage() can only convert textures of types PF_A8R8G8B8, PF_R8G8B8A8, and PF_B8G8R8A8"));
+		return nullptr;
+	}
+
+	int32 Width = Source->GetSizeX();
+	int32 Height = Source->GetSizeY();
+	const int32 NumComponents = 4;
+	const int32 BitsPerComponent = 8;
+	const int32 BitsPerPixel = NumComponents * BitsPerComponent;
+	const int32 ImageSizeInBytes = Width * Height * NumComponents;
+	const int32 BytesPerRow = Width * NumComponents;
+
+    CGBitmapInfo BitmapInfo = 0;
+	switch (SourceFormat)
+	{
+		case PF_A8R8G8B8:
+		{
+            BitmapInfo = kCGBitmapByteOrder32Big | kCGImageAlphaLast;
+			break;
+		}
+        case PF_R8G8B8A8:
+        {
+            BitmapInfo = kCGBitmapByteOrder32Big | kCGImageAlphaFirst;
+            break;
+        }
+        case PF_B8G8R8A8:
+        {
+            BitmapInfo = kCGBitmapByteOrder32Little | kCGImageAlphaLast;
+            break;
+        }
+	}
+
+	CGImageRef ImageRef = nullptr;
+
+	TArray<uint8*> MipPointers;
+	MipPointers.AddZeroed(Source->GetNumMips());
+	Source->GetMipData(0, (void**)MipPointers.GetData());
+	if (MipPointers[0] != nullptr)
+	{
+        // The data ref is going to free this result
+		CFDataRef DataRef = CFDataCreateWithBytesNoCopy(NULL, MipPointers[0], ImageSizeInBytes, kCFAllocatorDefault);
+		CGDataProviderRef DataProviderRef = CGDataProviderCreateWithCFData(DataRef);
+		CGColorSpaceRef ColorSpaceRef = CGColorSpaceCreateDeviceRGB();
+
+		ImageRef = CGImageCreate(Width, Height, BitsPerComponent, BitsPerPixel, BytesPerRow, ColorSpaceRef, BitmapInfo, DataProviderRef, nullptr, true, kCGRenderingIntentDefault);
+
+		// Release our temporary memory
+		CGColorSpaceRelease(ColorSpaceRef);
+		CGDataProviderRelease(DataProviderRef);
+		CFRelease(DataRef);
+	}
+	return ImageRef;
+}
+
+void FAppleImageUtilsPlugin::ConvertToJPEG(CIImage* SourceImage, TArray<uint8>& OutBytes, int32 Quality, bool bWantColor, bool bUseGpu, float Scale, ETextureRotationDirection Rotate)
+{
+	SCOPED_AUTORELEASE_POOL;
+
+	// Convert to Apple objects
+	CIContext* ConversionContext = [CIContext contextWithOptions: ToCpuDictionary(bUseGpu)];
+	CGColorSpaceRef ColorSpace = ToColorSpace(bWantColor);
+	CIImage* Image = ApplyScaleAndRotation(SourceImage, Scale, Rotate);
+
+	// This will perform the work on the GPU or inline on this thread
+	NSData* ConvertedData = [ConversionContext JPEGRepresentationOfImage: Image colorSpace: ColorSpace options: ToQualityDictionary(Quality)];
+	if (ConvertedData != nullptr)
+	{
+		uint32 CompressedSize = ConvertedData.length;
+		OutBytes.AddUninitialized(CompressedSize);
+		FPlatformMemory::Memcpy(OutBytes.GetData(), [ConvertedData bytes], CompressedSize);
+	}
+
+	CGColorSpaceRelease(ColorSpace);
+}
+
+#if SUPPORTS_IMAGE_UTILS_2_1
+void FAppleImageUtilsPlugin::ConvertToHEIF(CIImage* SourceImage, TArray<uint8>& OutBytes, int32 Quality,  bool bWantColor, bool bUseGpu, float Scale, ETextureRotationDirection Rotate)
+{
+	SCOPED_AUTORELEASE_POOL;
+
+	// Convert to Apple objects
+	CIContext* ConversionContext = [CIContext contextWithOptions: ToCpuDictionary(bUseGpu)];
+	CGColorSpaceRef ColorSpace = ToColorSpace(bWantColor);
+	CIImage* Image = ApplyScaleAndRotation(SourceImage, Scale, Rotate);
+
+	// This will perform the work on the GPU or inline on this thread
+	NSData* ConvertedData = [ConversionContext HEIFRepresentationOfImage: Image format: kCIFormatARGB8 colorSpace: ColorSpace options: ToQualityDictionary(Quality)];
+	if (ConvertedData != nullptr)
+	{
+		uint32 CompressedSize = ConvertedData.length;
+		OutBytes.AddUninitialized(CompressedSize);
+		FPlatformMemory::Memcpy(OutBytes.GetData(), [ConvertedData bytes], CompressedSize);
+	}
+
+	CGColorSpaceRelease(ColorSpace);
+}
+
+void FAppleImageUtilsPlugin::ConvertToPNG(CIImage* SourceImage, TArray<uint8>& OutBytes, bool bWantColor, bool bUseGpu, float Scale, ETextureRotationDirection Rotate)
+{
+	SCOPED_AUTORELEASE_POOL;
+
+	// Convert to Apple objects
+	CIContext* ConversionContext = [CIContext contextWithOptions: ToCpuDictionary(bUseGpu)];
+	CGColorSpaceRef ColorSpace = ToColorSpace(bWantColor);
+	CIImage* Image = ApplyScaleAndRotation(SourceImage, Scale, Rotate);
+
+	// This will perform the work on the GPU or inline on this thread
+	NSData* ConvertedData = [ConversionContext PNGRepresentationOfImage: Image format: kCIFormatARGB8 colorSpace: ColorSpace options: @{}];
+	if (ConvertedData != nullptr)
+	{
+		uint32 CompressedSize = ConvertedData.length;
+		OutBytes.AddUninitialized(CompressedSize);
+		FPlatformMemory::Memcpy(OutBytes.GetData(), [ConvertedData bytes], CompressedSize);
+	}
+
+	CGColorSpaceRelease(ColorSpace);
+}
+
+void FAppleImageUtilsPlugin::ConvertToTIFF(CIImage* SourceImage, TArray<uint8>& OutBytes, bool bWantColor, bool bUseGpu, float Scale, ETextureRotationDirection Rotate)
+{
+	SCOPED_AUTORELEASE_POOL;
+
+	// Convert to Apple objects
+	CIContext* ConversionContext = [CIContext contextWithOptions: ToCpuDictionary(bUseGpu)];
+	CGColorSpaceRef ColorSpace = ToColorSpace(bWantColor);
+	CIImage* Image = ApplyScaleAndRotation(SourceImage, Scale, Rotate);
+
+	// This will perform the work on the GPU or inline on this thread
+	NSData* ConvertedData = [ConversionContext TIFFRepresentationOfImage: Image format: kCIFormatARGB8 colorSpace: ColorSpace options: @{}];
+	if (ConvertedData != nullptr)
+	{
+		uint32 CompressedSize = ConvertedData.length;
+		OutBytes.AddUninitialized(CompressedSize);
+		FPlatformMemory::Memcpy(OutBytes.GetData(), [ConvertedData bytes], CompressedSize);
+	}
+
+	CGColorSpaceRelease(ColorSpace);
+}
+#endif
+
+	#pragma clang diagnostic pop
+#endif
+
+/** Dummy class needed to support Cast<IAppleImageInterface>(Object) */
+UAppleImageInterface::UAppleImageInterface(const FObjectInitializer& ObjectInitializer) :
+	Super(ObjectInitializer)
+{
+}
diff --git a/Engine/Plugins/Runtime/AppleImageUtils/Source/AppleImageUtils/Public/AppleImageUtilsAvailability.h b/Engine/Plugins/Runtime/AppleImageUtils/Source/AppleImageUtils/Public/AppleImageUtilsAvailability.h
new file mode 100644
--- /dev/null
+++ b/Engine/Plugins/Runtime/AppleImageUtils/Source/AppleImageUtils/Public/AppleImageUtilsAvailability.h
@@ -0,0 +1,134 @@
+// Copyright 1998-2019 Epic Games, Inc. All Rights Reserved.
+
+#pragma once
+
+#ifndef PLATFORM_IOS
+	#define PLATFORM_IOS 0
+#endif
+
+#ifndef PLATFORM_TVOS
+	#define PLATFORM_TVOS 0
+#endif
+
+#ifndef PLATFORM_MAC
+	#define PLATFORM_MAC 0
+#endif
+
+#if PLATFORM_IOS && !PLATFORM_TVOS
+
+	// Check for image utils 1.0
+	#ifdef __IPHONE_10_0
+		#define SUPPORTS_IMAGE_UTILS_1_0 1
+		#define RT_CHECK_IMAGE_UTILS_1_0 @available(iOS 10.0, *)
+	#else
+		#define SUPPORTS_IMAGE_UTILS_1_0 0
+	#endif
+
+	// Check for image utils 2
+	#ifdef __IPHONE_11_0
+		#define SUPPORTS_IMAGE_UTILS_2_0 1
+		#define RT_CHECK_IMAGE_UTILS_2_0 @available(iOS 11.0, *)
+	#else
+		#define SUPPORTS_IMAGE_UTILS_2_0 0
+	#endif
+
+	// Check for image utils 2.1
+	#ifdef __IPHONE_11_0
+		#define SUPPORTS_IMAGE_UTILS_2_1 1
+		#define RT_CHECK_IMAGE_UTILS_2_1 @available(iOS 11.0, *)
+	#else
+		#define SUPPORTS_IMAGE_UTILS_2_1 0
+	#endif
+
+#elif PLATFORM_MAC
+
+	// Check for image utils 1.0
+	#ifdef MAC_OS_X_VERSION_10_12
+		#define SUPPORTS_IMAGE_UTILS_1_0 1
+		#define RT_CHECK_IMAGE_UTILS_1_0 @available(macOS 10.12, *)
+	#else
+		#define SUPPORTS_IMAGE_UTILS_1_0 0
+	#endif
+
+	// Check for image utils 2
+	#ifdef MAC_OS_X_VERSION_10_13
+		#define SUPPORTS_IMAGE_UTILS_2_0 1
+		#define RT_CHECK_IMAGE_UTILS_2_0 @available(macOS 10.13, *)
+	#else
+		#define SUPPORTS_IMAGE_UTILS_2_0 0
+	#endif
+
+	// Check for image utils 2.1
+	#ifdef MAC_OS_X_VERSION_10_13
+        //@todo joeg -- enable again once you figure out the failure on 10.12.6
+		#define SUPPORTS_IMAGE_UTILS_2_1 0
+		#define RT_CHECK_IMAGE_UTILS_2_1 @available(macOS 10.13.4, *)
+	#else
+		#define SUPPORTS_IMAGE_UTILS_2_1 0
+	#endif
+
+#else
+
+	#define SUPPORTS_IMAGE_UTILS_1_0 0
+	#define SUPPORTS_IMAGE_UTILS_2_0 0
+	#define SUPPORTS_IMAGE_UTILS_2_1 0
+
+#endif
+
+class FAppleImageUtilsAvailability
+{
+public:
+	static bool SupportsImageUtils10()
+	{
+		static bool bSupportsImageUtils10 = false;
+#if SUPPORTS_IMAGE_UTILS_1_0
+		static bool bSupportChecked = false;
+		if (!bSupportChecked)
+		{
+			bSupportChecked = true;
+			// This call is slow, so cache the result
+			if (RT_CHECK_IMAGE_UTILS_1_0)
+			{
+				bSupportsImageUtils10 = true;
+			}
+		}
+#endif
+		return bSupportsImageUtils10;
+	}
+
+	static bool SupportsImageUtils20()
+	{
+		static bool bSupportsImageUtils20 = false;
+#if SUPPORTS_IMAGE_UTILS_2_0
+		static bool bSupportChecked = false;
+		if (!bSupportChecked)
+		{
+			bSupportChecked = true;
+			// This call is slow, so cache the result
+			if (RT_CHECK_IMAGE_UTILS_2_0)
+			{
+				bSupportsImageUtils20 = true;
+			}
+		}
+#endif
+		return bSupportsImageUtils20;
+	}
+
+	static bool SupportsImageUtils21()
+	{
+		static bool bSupportsImageUtils21 = false;
+#if SUPPORTS_IMAGE_UTILS_2_1
+		static bool bSupportChecked = false;
+		if (!bSupportChecked)
+		{
+			bSupportChecked = true;
+			// This call is slow, so cache the result
+			if (RT_CHECK_IMAGE_UTILS_2_1)
+			{
+				bSupportsImageUtils21 = true;
+			}
+		}
+#endif
+		return bSupportsImageUtils21;
+	}
+};
diff --git a/Engine/Plugins/Runtime/AppleImageUtils/Source/AppleImageUtils/Public/AppleImageUtilsBlueprintProxy.h b/Engine/Plugins/Runtime/AppleImageUtils/Source/AppleImageUtils/Public/AppleImageUtilsBlueprintProxy.h
new file mode 100644
--- /dev/null
+++ b/Engine/Plugins/Runtime/AppleImageUtils/Source/AppleImageUtils/Public/AppleImageUtilsBlueprintProxy.h
@@ -0,0 +1,112 @@
+// Copyright 1998-2019 Epic Games, Inc. All Rights Reserved.
+
+#pragma once
+
+#include "CoreMinimal.h"
+#include "UObject/ObjectMacros.h"
+#include "UObject/Object.h"
+#include "Engine/EngineTypes.h"
+#include "Tickable.h"
+
+#include "AppleImageUtilsTypes.h"
+
+#include "AppleImageUtilsBlueprintProxy.generated.h"
+
+class FAppleImageUtilsConversionTaskBase;
+
+USTRUCT(BlueprintType)
+struct FAppleImageUtilsImageConversionResult
+{
+	GENERATED_USTRUCT_BODY()
+
+	UPROPERTY(BlueprintReadOnly, Category="Image Conversion")
+	FString Error;
+
+	UPROPERTY(BlueprintReadOnly, Category="Image Conversion")
+	TArray<uint8> ImageData;
+};
+
+DECLARE_DYNAMIC_MULTICAST_DELEGATE_OneParam(FAppleImageConversionDelegate, const FAppleImageUtilsImageConversionResult&, ConversionResult);
+
+UCLASS(MinimalAPI)
+class UAppleImageUtilsBaseAsyncTaskBlueprintProxy :
+	public UObject,
+	public FTickableGameObject
+{
+	GENERATED_UCLASS_BODY()
+
+public:
+	UPROPERTY(BlueprintAssignable)
+	FAppleImageConversionDelegate OnSuccess;
+
+	UPROPERTY(BlueprintAssignable)
+	FAppleImageConversionDelegate OnFailure;
+
+	/**
+	 * Converts a image to an array of JPEG data in a background task
+	 *
+	 * @param SourceImage the image to compress
+	 * @param Quality the quality level to compress to
+	 * @param bWantColor whether the JPEG is color (true) or monochrome (false)
+	 * @param bUseGpu whether to use the GPU (true) or the CPU (false) to compress
+	 * @param Scale whether to scale the image before conversion, defaults to no scaling
+	 * @param Rotate a direction to rotate the image in during conversion, defaults to none
+	 */
+	UFUNCTION(BlueprintCallable, Meta=(BlueprintInternalUseOnly="true", DisplayName="Convert To JPEG"), Category="Image Conversion")
+	static UAppleImageUtilsBaseAsyncTaskBlueprintProxy* CreateProxyObjectForConvertToJPEG(UTexture* SourceImage, int32 Quality = 85, bool bWantColor = true, bool bUseGpu = true, float Scale = 1.f, ETextureRotationDirection Rotate = ETextureRotationDirection::None);
+
+	/**
+	 * Converts a image to an array of HEIF data in a background task
+	 *
+	 * @param SourceImage the image to compress
+	 * @param Quality the quality level to compress to
+	 * @param bWantColor whether the HEIF is color (true) or monochrome (false)
+	 * @param bUseGpu whether to use the GPU (true) or the CPU (false) to compress
+	 * @param Scale whether to scale the image before conversion, defaults to no scaling
+	 * @param Rotate a direction to rotate the image in during conversion, defaults to none
+	 */
+	UFUNCTION(BlueprintCallable, Meta=(BlueprintInternalUseOnly="true", DisplayName="Convert To HEIF"), Category="Image Conversion")
+	static UAppleImageUtilsBaseAsyncTaskBlueprintProxy* CreateProxyObjectForConvertToHEIF(UTexture* SourceImage, int32 Quality = 85, bool bWantColor = true, bool bUseGpu = true, float Scale = 1.f, ETextureRotationDirection Rotate = ETextureRotationDirection::None);
+
+	/**
+	 * Converts a image to an array of TIFF data in a background task
+	 *
+	 * @param SourceImage the image to compress
+	 * @param Quality the quality level to compress to
+	 * @param bWantColor whether the TIFF is color (true) or monochrome (false)
+	 * @param bUseGpu whether to use the GPU (true) or the CPU (false) to compress
+	 * @param Scale whether to scale the image before conversion, defaults to no scaling
+	 * @param Rotate a direction to rotate the image in during conversion, defaults to none
+	 */
+	UFUNCTION(BlueprintCallable, Meta=(BlueprintInternalUseOnly="true", DisplayName="Convert To TIFF"), Category="Image Conversion")
+	static UAppleImageUtilsBaseAsyncTaskBlueprintProxy* CreateProxyObjectForConvertToTIFF(UTexture* SourceImage, bool bWantColor = true, bool bUseGpu = true, float Scale = 1.f, ETextureRotationDirection Rotate = ETextureRotationDirection::None);
+
+	/**
+	 * Converts a image to an array of TIFF data in a background task
+	 *
+	 * @param SourceImage the image to compress
+	 * @param Quality the quality level to compress to
+	 * @param bWantColor whether the PNG is color (true) or monochrome (false)
+	 * @param bUseGpu whether to use the GPU (true) or the CPU (false) to compress
+	 * @param Scale whether to scale the image before conversion, defaults to no scaling
+	 * @param Rotate a direction to rotate the image in during conversion, defaults to none
+	 */
+	UFUNCTION(BlueprintCallable, Meta=(BlueprintInternalUseOnly="true", DisplayName="Convert To PNG"), Category="Image Conversion")
+	static UAppleImageUtilsBaseAsyncTaskBlueprintProxy* CreateProxyObjectForConvertToPNG(UTexture* SourceImage, bool bWantColor = true, bool bUseGpu = true, float Scale = 1.f, ETextureRotationDirection Rotate = ETextureRotationDirection::None);
+
+	//~ Begin FTickableObject Interface
+	virtual void Tick(float DeltaTime) override;
+	virtual bool IsTickable() const override { return bShouldTick; }
+	virtual TStatId GetStatId() const override { RETURN_QUICK_DECLARE_CYCLE_STAT(UAppleImageUtilsBaseAsyncTaskBlueprintProxy, STATGROUP_Tickables); }
+	//~ End FTickableObject Interface
+
+	/** The async task to check during Tick() */
+	TSharedPtr<FAppleImageUtilsConversionTaskBase, ESPMode::ThreadSafe> ConversionTask;
+
+	UPROPERTY(BlueprintReadOnly, Category="Image Conversion")
+	FAppleImageUtilsImageConversionResult ConversionResult;
+
+private:
+	/** True until the async task completes, then false */
+	bool bShouldTick;
+};
diff --git a/Engine/Plugins/Runtime/AppleImageUtils/Source/AppleImageUtils/Public/AppleImageUtilsTypes.h b/Engine/Plugins/Runtime/AppleImageUtils/Source/AppleImageUtils/Public/AppleImageUtilsTypes.h
new file mode 100644
--- /dev/null
+++ b/Engine/Plugins/Runtime/AppleImageUtils/Source/AppleImageUtils/Public/AppleImageUtilsTypes.h
@@ -0,0 +1,62 @@
+// Copyright 1998-2019 Epic Games, Inc. All Rights Reserved.
+
+#pragma once
+
+#include "CoreMinimal.h"
+#include "UObject/ObjectMacros.h"
+#include "UObject/Interface.h"
+#if PLATFORM_MAC || PLATFORM_IOS
+	#import <CoreVideo/CoreVideo.h>
+	#import <CoreImage/CoreImage.h>
+#endif
+#include "AppleImageUtilsTypes.generated.h"
+
+UENUM(BlueprintType, Category="AppleImageUtils", meta=(Experimental))
+enum class ETextureRotationDirection : uint8
+{
+	None,
+	Left,
+	Right,
+	Down
+};
+
+UENUM(BlueprintType, Category="AppleImageUtils", meta=(Experimental))
+enum class EAppleTextureType : uint8
+{
+	Unknown,
+	Image,
+	PixelBuffer,
+	Surface,
+	MetalTexture
+};
+
+/** Dummy class needed to support Cast<IAppleImageInterface>(Object). */
+UINTERFACE()
+class APPLEIMAGEUTILS_API UAppleImageInterface : public UInterface
+{
+	GENERATED_UINTERFACE_BODY()
+};
+
+/**
+ * Base class for accessing the raw Apple image data
+ */
+class APPLEIMAGEUTILS_API IAppleImageInterface
+{
+	GENERATED_IINTERFACE_BODY()
+
+public:
+	/** @return the type of image held by the implementing object */
+	virtual EAppleTextureType GetTextureType() const = 0;
+
+#if PLATFORM_MAC || PLATFORM_IOS
+	/** @return the CIImage held by the implementing object */
+	virtual CIImage* GetImage() const { return nullptr; }
+	/** @return the CVPixelBuffer held by the implementing object */
+	virtual CVPixelBufferRef GetPixelBuffer() const { return nullptr; }
+	/** @return the IOSurface held by the implementing object */
+	virtual IOSurfaceRef GetSurface() const { return nullptr; }
+
+	/** @return the MTLTextureid<MTL held by the implementing object */
+	virtual id<MTLTexture> GetMetalTexture() const { return nullptr; }
+#endif
+};
diff --git a/Engine/Plugins/Runtime/AppleImageUtils/Source/AppleImageUtils/Public/IAppleImageUtilsPlugin.h b/Engine/Plugins/Runtime/AppleImageUtils/Source/AppleImageUtils/Public/IAppleImageUtilsPlugin.h
new file mode 100644
--- /dev/null
+++ b/Engine/Plugins/Runtime/AppleImageUtils/Source/AppleImageUtils/Public/IAppleImageUtilsPlugin.h
@@ -0,0 +1,207 @@
+// Copyright 1998-2019 Epic Games, Inc. All Rights Reserved.
+
+#pragma once
+
+#include "CoreMinimal.h"
+#include "Modules/ModuleInterface.h"
+#include "Modules/ModuleManager.h"
+#include "HAL/ThreadSafeBool.h"
+
+#include "AppleImageUtilsTypes.h"
+
+#include "AppleImageUtilsAvailability.h"
+
+#if SUPPORTS_IMAGE_UTILS_1_0
+	#import <CoreGraphics/CoreGraphics.h>
+#endif
+
+class UTexture;
+class UTexture2D;
+
+/**
+ * Interface to access the async compression request
+ */
+class APPLEIMAGEUTILS_API IAppleImageUtilsConversionTask
+{
+public:
+	/** @return whether the task succeeded or not */
+	virtual bool HadError() const = 0;
+	/** @return information about the error if there was one */
+	virtual FString GetErrorReason() const = 0;
+	/** @return whether the task has completed or not */
+	virtual bool IsDone() const = 0;
+	/** @return how long it took this task to complete */
+	virtual float GetElapsedTime() const = 0;
+	/** @return the data once the task has completed or empty array if still in progress. NOTE: this should use MoveTemp() to avoid duplication so call once :) */
+	virtual TArray<uint8> GetData() = 0;
+};
+
+/**
+ * Base class for implementing IAppleImageUtilsCompressionTasks
+ */
+class APPLEIMAGEUTILS_API FAppleImageUtilsConversionTaskBase :
+	public IAppleImageUtilsConversionTask
+{
+protected:
+	FAppleImageUtilsConversionTaskBase() {}
+	virtual ~FAppleImageUtilsConversionTaskBase() {}
+
+public:
+	virtual bool IsDone() const override { return bIsDone; }
+	virtual bool HadError() const override { return bHadError; }
+	virtual FString GetErrorReason() const override { return Error; }
+	virtual float GetElapsedTime() const override { return (EndTime - StartTime) * 1000.f; }
+
+protected:
+	FThreadSafeBool bIsDone;
+	FThreadSafeBool bHadError;
+	double StartTime = 0.f;
+	double EndTime = 0.f;
+	FString Error;
+};
+
+class APPLEIMAGEUTILS_API IAppleImageUtilsPlugin :
+	public IModuleInterface
+{
+public:
+	/**
+	 * Singleton-like access to this module's interface.  This is just for convenience!
+	 * Beware of calling this during the shutdown phase, though.  Your module might have been unloaded already.
+	 *
+	 * @return Returns singleton instance, loading the module on demand if needed
+	 */
+	static inline IAppleImageUtilsPlugin& Get()
+	{
+		return FModuleManager::LoadModuleChecked<IAppleImageUtilsPlugin>("AppleImageUtils");
+	}
+
+	/**
+	 * Checks to see if this module is loaded and ready.  It is only valid to call Get() if IsAvailable() returns true.
+	 *
+	 * @return True if the module is loaded and ready to use
+	 */
+	static inline bool IsAvailable()
+	{
+		return FModuleManager::Get().IsModuleLoaded("AppleImageUtils");
+	}
+
+	/**
+	 * Converts a image to an array of JPEG data in a background task
+	 *
+	 * @param SourceImage the image to compress (NOTE: must support UAppleImageInterface)
+	 * @param Quality the quality level to compress to
+	 * @param bWantColor whether the JPEG is color (true) or monochrome (false)
+	 * @param bUseGpu whether to use the GPU (true) or the CPU (false) to compress
+	 * @param Scale whether to scale the image before converting, defaults to no scale operation
+	 * @param Rotate a direction to rotate the image in during conversion, defaults to none
+	 *
+	 * @return the async task that is doing the conversion
+	 */
+	virtual TSharedPtr<FAppleImageUtilsConversionTaskBase, ESPMode::ThreadSafe> ConvertToJPEG(UTexture* SourceImage, int32 Quality = 85, bool bWantColor = true, bool bUseGpu = true, float Scale = 1.f, ETextureRotationDirection Rotate = ETextureRotationDirection::None) = 0;
+
+	/**
+	 * Converts a image to an array of HEIF data in a background task
+	 *
+	 * @param SourceImage the image to compress (NOTE: must support UAppleImageInterface)
+	 * @param Quality the quality level to compress to
+	 * @param bWantColor whether the HEIF is color (true) or monochrome (false)
+	 * @param bUseGpu whether to use the GPU (true) or the CPU (false) to compress
+	 * @param Scale whether to scale the image before converting, defaults to no scale operation
+	 * @param Rotate a direction to rotate the image in during conversion, defaults to none
+	 *
+	 * @return the async task that is doing the conversion
+	 */
+	virtual TSharedPtr<FAppleImageUtilsConversionTaskBase, ESPMode::ThreadSafe> ConvertToHEIF(UTexture* SourceImage, int32 Quality = 85,  bool bWantColor = true, bool bUseGpu = true, float Scale = 1.f, ETextureRotationDirection Rotate = ETextureRotationDirection::None) = 0;
+
+	/**
+	 * Converts a image to an array of PNG data in a background task
+	 *
+	 * @param SourceImage the image to compress (NOTE: must support UAppleImageInterface)
+	 * @param Quality the quality level to compress to
+	 * @param bWantColor whether the PNG is color (true) or monochrome (false)
+	 * @param bUseGpu whether to use the GPU (true) or the CPU (false) to compress
+	 * @param Scale whether to scale the image before converting, defaults to no scale operation
+	 * @param Rotate a direction to rotate the image in during conversion, defaults to none
+	 *
+	 * @return the async task that is doing the conversion
+	 */
+	virtual TSharedPtr<FAppleImageUtilsConversionTaskBase, ESPMode::ThreadSafe> ConvertToPNG(UTexture* SourceImage, bool bWantColor = true, bool bUseGpu = true, float Scale = 1.f, ETextureRotationDirection Rotate = ETextureRotationDirection::None) = 0;
+
+	/**
+	 * Converts a image to an array of TIFF data in a background task
+	 *
+	 * @param SourceImage the image to compress (NOTE: must support UAppleImageInterface)
+	 * @param Quality the quality level to compress to
+	 * @param bUseGpu whether to use the GPU (true) or the CPU (false) to compress
+	 * @param Scale whether to scale the image before converting, defaults to no scale operation
+	 * @param Rotate a direction to rotate the image in during conversion, defaults to none
+	 *
+	 * @return the async task that is doing the conversion
+	 */
+	virtual TSharedPtr<FAppleImageUtilsConversionTaskBase, ESPMode::ThreadSafe> ConvertToTIFF(UTexture* SourceImage, bool bWantColor = true, bool bUseGpu = true, float Scale = 1.f, ETextureRotationDirection Rotate = ETextureRotationDirection::None) = 0;
+
+#if SUPPORTS_IMAGE_UTILS_1_0
+	/**
+	 * Copies the contents of a UTexture2D to a CGImage object
+	 *
+	 * @param SourceImage the image to create a CGImage from
+	 *
+	 * @return the CGImage if conversion was possible
+	 */
+	virtual CGImageRef UTexture2DToCGImage(UTexture2D* Source) = 0;
+
+	/**
+	 * Converts a image to an array of JPEG data synchronously
+	 *
+	 * @param SourceImage the image to compress
+	 * @param OutBytes the buffer that is populated during compression
+	 * @param Quality the quality level to compress to
+	 * @param bWantColor whether the JPEG is color (true) or monochrome (false)
+	 * @param bUseGpu whether to use the GPU (true) or the CPU (false) to compress
+	 * @param Scale whether to scale the image before converting, defaults to no scale operation
+	 * @param Rotate a direction to rotate the image in during conversion, defaults to none
+	 */
+	virtual void ConvertToJPEG(CIImage* SourceImage, TArray<uint8>& OutBytes, int32 Quality = 85, bool bWantColor = true, bool bUseGpu = true, float Scale = 1.f, ETextureRotationDirection Rotate = ETextureRotationDirection::None) = 0;
+
+#if SUPPORTS_IMAGE_UTILS_2_1
+	/**
+	 * Converts a image to an array of HEIF data synchronously
+	 *
+	 * @param SourceImage the image to compress
+	 * @param OutBytes the buffer that is populated during compression
+	 * @param Quality the quality level to compress to
+	 * @param bWantColor whether the HEIF is color (true) or monochrome (false)
+	 * @param bUseGpu whether to use the GPU (true) or the CPU (false) to compress
+	 * @param Scale whether to scale the image before converting, defaults to no scale operation
+	 * @param Rotate a direction to rotate the image in during conversion, defaults to none
+	 */
+	virtual void ConvertToHEIF(CIImage* SourceImage, TArray<uint8>& OutBytes, int32 Quality = 85,  bool bWantColor = true, bool bUseGpu = true, float Scale = 1.f, ETextureRotationDirection Rotate = ETextureRotationDirection::None) = 0;
+
+	/**
+	 * Converts a image to an array of PNG data synchronously
+	 *
+	 * @param SourceImage the image to compress
+	 * @param OutBytes the buffer that is populated during compression
+	 * @param Quality the quality level to compress to
+	 * @param bWantColor whether the PNG is color (true) or monochrome (false)
+	 * @param bUseGpu whether to use the GPU (true) or the CPU (false) to compress
+	 * @param Scale whether to scale the image before converting, defaults to no scale operation
+	 * @param Rotate a direction to rotate the image in during conversion, defaults to none
+	 */
+	virtual void ConvertToPNG(CIImage* SourceImage, TArray<uint8>& OutBytes, bool bWantColor = true, bool bUseGpu = true, float Scale = 1.f, ETextureRotationDirection Rotate = ETextureRotationDirection::None) = 0;
+
+	/**
+	 * Converts a image to an array of TIFF data synchronously
+	 *
+	 * @param SourceImage the image to compress
+	 * @param OutBytes the buffer that is populated during compression
+	 * @param Quality the quality level to compress to
+	 * @param bUseGpu whether to use the GPU (true) or the CPU (false) to compress
+	 * @param Scale whether to scale the image before converting, defaults to no scale operation
+	 * @param Rotate a direction to rotate the image in during conversion, defaults to none
+	 */
+	virtual void ConvertToTIFF(CIImage* SourceImage, TArray<uint8>& OutBytes, bool bWantColor = true, bool bUseGpu = true, float Scale = 1.f, ETextureRotationDirection Rotate = ETextureRotationDirection::None) = 0;
+#endif
+#endif
+};
+
diff --git a/Engine/Plugins/Runtime/AppleImageUtils/Source/AppleImageUtilsBlueprintSupport/AppleImageUtilsBlueprintSupport.Build.cs b/Engine/Plugins/Runtime/AppleImageUtils/Source/AppleImageUtilsBlueprintSupport/AppleImageUtilsBlueprintSupport.Build.cs
new file mode 100644
--- /dev/null
+++ b/Engine/Plugins/Runtime/AppleImageUtils/Source/AppleImageUtilsBlueprintSupport/AppleImageUtilsBlueprintSupport.Build.cs
@@ -0,0 +1,59 @@
+// Copyright 1998-2019 Epic Games, Inc. All Rights Reserved.
+
+namespace UnrealBuildTool.Rules
+{
+	public class AppleImageUtilsBlueprintSupport : ModuleRules
+	{
+		public AppleImageUtilsBlueprintSupport(ReadOnlyTargetRules Target) : base(Target)
+		{
+            PublicIncludePaths.AddRange(
+                new string[] {
+                    // ... add public include paths required here ...
+                }
+                );
+
+
+            PrivateIncludePaths.AddRange(
+                new string[] {
+                    "AppleImageUtils/Public",
+                    "AppleImageUtilsBlueprintSupport/Classes"
+                    // ... add other private include paths required here ...
+                }
+                );
+
+
+            PublicDependencyModuleNames.AddRange(
+                new string[]
+                {
+                    "Core",
+                    "Engine",
+                    "BlueprintGraph",
+                    "AppleImageUtils"
+                    // ... add other public dependencies that you statically link with here ...
+                }
+                );
+
+            PublicFrameworks.AddRange(
+                new string[]
+                {
+                    // ... add other public dependencies that you statically link with here ...
+                }
+                );
+
+			PrivateDependencyModuleNames.AddRange(
+				new string[]
+				{
+    				"CoreUObject"
+					// ... add private dependencies that you statically link with here ...
+				}
+				);
+
+			DynamicallyLoadedModuleNames.AddRange(
+				new string[]
+				{
+					// ... add any modules that your module loads dynamically here ...
+				}
+				);
+		}
+	}
+}
diff --git a/Engine/Plugins/Runtime/AppleImageUtils/Source/AppleImageUtilsBlueprintSupport/Classes/AppleImageUtilsBlueprintSupport.h b/Engine/Plugins/Runtime/AppleImageUtils/Source/AppleImageUtilsBlueprintSupport/Classes/AppleImageUtilsBlueprintSupport.h
new file mode 100644
--- /dev/null
+++ b/Engine/Plugins/Runtime/AppleImageUtils/Source/AppleImageUtilsBlueprintSupport/Classes/AppleImageUtilsBlueprintSupport.h
@@ -0,0 +1,37 @@
+// Copyright 1998-2019 Epic Games, Inc. All Rights Reserved.
+
+#pragma once
+
+#include "CoreMinimal.h"
+#include "UObject/ObjectMacros.h"
+#include "K2Node_BaseAsyncTask.h"
+
+#include "AppleImageUtilsBlueprintSupport.generated.h"
+
+UCLASS()
+class APPLEIMAGEUTILSBLUEPRINTSUPPORT_API UK2Node_ConvertToJPEG :
+	public UK2Node_BaseAsyncTask
+{
+	GENERATED_UCLASS_BODY()
+};
+
+UCLASS()
+class APPLEIMAGEUTILSBLUEPRINTSUPPORT_API UK2Node_ConvertToHEIF :
+	public UK2Node_BaseAsyncTask
+{
+	GENERATED_UCLASS_BODY()
+};
+
+UCLASS()
+class APPLEIMAGEUTILSBLUEPRINTSUPPORT_API UK2Node_ConvertToTIFF :
+	public UK2Node_BaseAsyncTask
+{
+	GENERATED_UCLASS_BODY()
+};
+
+UCLASS()
+class APPLEIMAGEUTILSBLUEPRINTSUPPORT_API UK2Node_ConvertToPNG :
+	public UK2Node_BaseAsyncTask
+{
+	GENERATED_UCLASS_BODY()
+};
diff --git a/Engine/Plugins/Runtime/AppleImageUtils/Source/AppleImageUtilsBlueprintSupport/Private/AppleImageUtilsBlueprintSupport.cpp b/Engine/Plugins/Runtime/AppleImageUtils/Source/AppleImageUtilsBlueprintSupport/Private/AppleImageUtilsBlueprintSupport.cpp
new file mode 100644
--- /dev/null
+++ b/Engine/Plugins/Runtime/AppleImageUtils/Source/AppleImageUtilsBlueprintSupport/Private/AppleImageUtilsBlueprintSupport.cpp
@@ -0,0 +1,50 @@
+// Copyright 1998-2019 Epic Games, Inc. All Rights Reserved.
+
+#include "AppleImageUtilsBlueprintSupport.h"
+#include "AppleImageUtilsBlueprintProxy.h"
+#include "Modules/ModuleInterface.h"
+#include "Modules/ModuleManager.h"
+
+
+class FAppleImageUtilsBlueprintSupport :
+	public IModuleInterface
+{
+	/** IModuleInterface implementation */
+	virtual void StartupModule() override {}
+	virtual void ShutdownModule() override {}
+
+};
+
+IMPLEMENT_MODULE(FAppleImageUtilsBlueprintSupport, AppleImageUtilsBlueprintSupport);
+
+UK2Node_ConvertToJPEG::UK2Node_ConvertToJPEG(const FObjectInitializer& ObjectInitializer)
+	: Super(ObjectInitializer)
+{
+	ProxyFactoryFunctionName = GET_FUNCTION_NAME_CHECKED(UAppleImageUtilsBaseAsyncTaskBlueprintProxy, CreateProxyObjectForConvertToJPEG);
+	ProxyFactoryClass = UAppleImageUtilsBaseAsyncTaskBlueprintProxy::StaticClass();
+	ProxyClass = UAppleImageUtilsBaseAsyncTaskBlueprintProxy::StaticClass();
+}
+
+UK2Node_ConvertToHEIF::UK2Node_ConvertToHEIF(const FObjectInitializer& ObjectInitializer)
+	: Super(ObjectInitializer)
+{
+	ProxyFactoryFunctionName = GET_FUNCTION_NAME_CHECKED(UAppleImageUtilsBaseAsyncTaskBlueprintProxy, CreateProxyObjectForConvertToHEIF);
+	ProxyFactoryClass = UAppleImageUtilsBaseAsyncTaskBlueprintProxy::StaticClass();
+	ProxyClass = UAppleImageUtilsBaseAsyncTaskBlueprintProxy::StaticClass();
+}
+
+UK2Node_ConvertToTIFF::UK2Node_ConvertToTIFF(const FObjectInitializer& ObjectInitializer)
+	: Super(ObjectInitializer)
+{
+	ProxyFactoryFunctionName = GET_FUNCTION_NAME_CHECKED(UAppleImageUtilsBaseAsyncTaskBlueprintProxy, CreateProxyObjectForConvertToTIFF);
+	ProxyFactoryClass = UAppleImageUtilsBaseAsyncTaskBlueprintProxy::StaticClass();
+	ProxyClass = UAppleImageUtilsBaseAsyncTaskBlueprintProxy::StaticClass();
+}
+
+UK2Node_ConvertToPNG::UK2Node_ConvertToPNG(const FObjectInitializer& ObjectInitializer)
+	: Super(ObjectInitializer)
+{
+	ProxyFactoryFunctionName = GET_FUNCTION_NAME_CHECKED(UAppleImageUtilsBaseAsyncTaskBlueprintProxy, CreateProxyObjectForConvertToPNG);
+	ProxyFactoryClass = UAppleImageUtilsBaseAsyncTaskBlueprintProxy::StaticClass();
+	ProxyClass = UAppleImageUtilsBaseAsyncTaskBlueprintProxy::StaticClass();
+}
