Removed the beta tag from the AppleImageTools plugin

diff --git a/Engine/Plugins/Runtime/AppleImageUtils/AppleImageUtils.uplugin b/Engine/Plugins/Runtime/AppleImageUtils/AppleImageUtils.uplugin
--- a/Engine/Plugins/Runtime/AppleImageUtils/AppleImageUtils.uplugin
+++ b/Engine/Plugins/Runtime/AppleImageUtils/AppleImageUtils.uplugin
@@ -12,7 +12,7 @@
 	"SupportURL" : "",
 	"EnabledByDefault" : true,
 	"CanContainContent" : false,
-	"IsBetaVersion" : true,
+	"IsBetaVersion" : false,
 	"Installed" : false,
 	"Modules" :
 	[
