Fixed a call to LoadModuleChecked while not on the game thread

diff --git a/Engine/Plugins/Runtime/AppleImageUtils/Source/AppleImageUtils/Public/IAppleImageUtilsPlugin.h b/Engine/Plugins/Runtime/AppleImageUtils/Source/AppleImageUtils/Public/IAppleImageUtilsPlugin.h
--- a/Engine/Plugins/Runtime/AppleImageUtils/Source/AppleImageUtils/Public/IAppleImageUtilsPlugin.h
+++ b/Engine/Plugins/Runtime/AppleImageUtils/Source/AppleImageUtils/Public/IAppleImageUtilsPlugin.h
@@ -64,17 +64,29 @@ class APPLEIMAGEUTILS_API IAppleImageUtilsPlugin :
 	public IModuleInterface
 {
 public:
+
 	/**
 	 * Singleton-like access to this module's interface.  This is just for convenience!
 	 * Beware of calling this during the shutdown phase, though.  Your module might have been unloaded already.
 	 *
 	 * @return Returns singleton instance, loading the module on demand if needed
 	 */
-	static inline IAppleImageUtilsPlugin& Get()
+	static inline IAppleImageUtilsPlugin& Load()
 	{
 		return FModuleManager::LoadModuleChecked<IAppleImageUtilsPlugin>("AppleImageUtils");
 	}
 
+	/**
+	 * Singleton-like access to this module's interface.  This is just for convenience!
+	 * Beware of calling this during the shutdown phase, though.  Your module might have been unloaded already.
+	 *
+	 * @return Returns singleton instance
+	 */
+	static inline IAppleImageUtilsPlugin& Get()
+	{
+		return FModuleManager::GetModuleChecked<IAppleImageUtilsPlugin>("AppleImageUtils");
+	}
+
 	/**
 	 * Checks to see if this module is loaded and ready.  It is only valid to call Get() if IsAvailable() returns true.
 	 *
