Removed the tvos blacklist for AppleImageUtils plugin

diff --git a/Engine/Plugins/Runtime/AppleImageUtils/AppleImageUtils.uplugin b/Engine/Plugins/Runtime/AppleImageUtils/AppleImageUtils.uplugin
--- a/Engine/Plugins/Runtime/AppleImageUtils/AppleImageUtils.uplugin
+++ b/Engine/Plugins/Runtime/AppleImageUtils/AppleImageUtils.uplugin
@@ -19,11 +19,7 @@
 		{
 			"Name" : "AppleImageUtils",
 			"Type" : "Runtime",
-			"LoadingPhase" : "Default",
-			"BlacklistPlatforms" :
-			[
-				"TVOS"
-			]
+			"LoadingPhase" : "Default"
 		},
 		{
 			"Name" : "AppleImageUtilsBlueprintSupport",
