Allowing multithreaded destruction of FTickableGameObject UObjects

diff --git a/Engine/Plugins/Experimental/AppleImageUtils/Source/AppleImageUtils/Public/AppleImageUtilsBlueprintProxy.h b/Engine/Plugins/Experimental/AppleImageUtils/Source/AppleImageUtils/Public/AppleImageUtilsBlueprintProxy.h
--- a/Engine/Plugins/Experimental/AppleImageUtils/Source/AppleImageUtils/Public/AppleImageUtilsBlueprintProxy.h
+++ b/Engine/Plugins/Experimental/AppleImageUtils/Source/AppleImageUtils/Public/AppleImageUtilsBlueprintProxy.h
@@ -100,10 +100,6 @@ public:
 	virtual TStatId GetStatId() const override { RETURN_QUICK_DECLARE_CYCLE_STAT(UAppleImageUtilsBaseAsyncTaskBlueprintProxy, STATGROUP_Tickables); }
 	//~ End FTickableObject Interface
 
-	//~ Begin UObject Interface
-	virtual bool IsDestructionThreadSafe() const override { return false; }
-	//~ End UObject Interface
-
 	/** The async task to check during Tick() */
 	TSharedPtr<FAppleImageUtilsConversionTaskBase, ESPMode::ThreadSafe> ConversionTask;
 
diff --git a/Engine/Plugins/Experimental/AppleVision/Source/AppleVision/Public/AppleVisionBlueprintProxy.h b/Engine/Plugins/Experimental/AppleVision/Source/AppleVision/Public/AppleVisionBlueprintProxy.h
--- a/Engine/Plugins/Experimental/AppleVision/Source/AppleVision/Public/AppleVisionBlueprintProxy.h
+++ b/Engine/Plugins/Experimental/AppleVision/Source/AppleVision/Public/AppleVisionBlueprintProxy.h
@@ -44,10 +44,6 @@ public:
 	virtual TStatId GetStatId() const override { RETURN_QUICK_DECLARE_CYCLE_STAT(UAppleVisionDetectFacesAsyncTaskBlueprintProxy, STATGROUP_Tickables); }
 	//~ End FTickableObject Interface
 
-	//~ Begin UObject Interface
-	virtual bool IsDestructionThreadSafe() const override { return false; }
-	//~ End UObject Interface
-
 	/** The async task to check during Tick() */
 	TSharedPtr<FAppleVisionDetectFacesAsyncTaskBase, ESPMode::ThreadSafe> AsyncTask;
 
diff --git a/Engine/Plugins/Experimental/StereoPanorama/Source/StereoPanorama/Private/SceneCapturer.h b/Engine/Plugins/Experimental/StereoPanorama/Source/StereoPanorama/Private/SceneCapturer.h
--- a/Engine/Plugins/Experimental/StereoPanorama/Source/StereoPanorama/Private/SceneCapturer.h
+++ b/Engine/Plugins/Experimental/StereoPanorama/Source/StereoPanorama/Private/SceneCapturer.h
@@ -72,10 +72,6 @@ public:
 		RETURN_QUICK_DECLARE_CYCLE_STAT( USceneCapturer, STATGROUP_Tickables );
 	}
 
-	//~ Begin UObject Interface
-	virtual bool IsDestructionThreadSafe() const override { return false; }
-	//~ End UObject Interface
-
 public:
 
 	void InitCaptureComponent( USceneCaptureComponent2D* CaptureComponent, float HFov, float VFov, EStereoscopicPass InStereoPass );
diff --git a/Engine/Plugins/Lumin/MagicLeap/Source/MagicLeapIdentity/Public/MagicLeapIdentity.h b/Engine/Plugins/Lumin/MagicLeap/Source/MagicLeapIdentity/Public/MagicLeapIdentity.h
--- a/Engine/Plugins/Lumin/MagicLeap/Source/MagicLeapIdentity/Public/MagicLeapIdentity.h
+++ b/Engine/Plugins/Lumin/MagicLeap/Source/MagicLeapIdentity/Public/MagicLeapIdentity.h
@@ -88,10 +88,6 @@ public:
 	/** FTickableGameObject interface */
 	virtual UWorld* GetTickableGameObjectWorld() const override;
 
-	//~ Begin UObject Interface
-	virtual bool IsDestructionThreadSafe() const override { return false; }
-	//~ End UObject Interface
-
 private:
 	class FIdentityImpl *Impl;
 };
diff --git a/Engine/Plugins/NetcodeUnitTest/NetcodeUnitTest/Source/NetcodeUnitTest/Classes/UnitTestManager.h b/Engine/Plugins/NetcodeUnitTest/NetcodeUnitTest/Source/NetcodeUnitTest/Classes/UnitTestManager.h
--- a/Engine/Plugins/NetcodeUnitTest/NetcodeUnitTest/Source/NetcodeUnitTest/Classes/UnitTestManager.h
+++ b/Engine/Plugins/NetcodeUnitTest/NetcodeUnitTest/Source/NetcodeUnitTest/Classes/UnitTestManager.h
@@ -297,11 +297,6 @@ public:
 	 */
 	bool Exec(UWorld* InWorld, const TCHAR* Cmd, FOutputDevice& Ar);
 
-	/**
-	 * UObject methods
-	 */
-	virtual bool IsDestructionThreadSafe() const override { return false; }
-	
 	/**
 	 * FOutputDevice methods
 	 */
diff --git a/Engine/Plugins/Runtime/Oculus/OculusVR/Source/OculusHMD/Private/OculusSceneCaptureCubemap.h b/Engine/Plugins/Runtime/Oculus/OculusVR/Source/OculusHMD/Private/OculusSceneCaptureCubemap.h
--- a/Engine/Plugins/Runtime/Oculus/OculusVR/Source/OculusHMD/Private/OculusSceneCaptureCubemap.h
+++ b/Engine/Plugins/Runtime/Oculus/OculusVR/Source/OculusHMD/Private/OculusSceneCaptureCubemap.h
@@ -56,9 +56,6 @@ public:
 	static void CaptureCubemapCommandHandler(const TArray<FString>& Args, UWorld* World, FOutputDevice& Ar);
 #endif // UE_BUILD_SHIPPING
 
-	// UObject implementation
-	virtual bool IsDestructionThreadSafe() const override { return false; }
-
 private:
 
 	enum EStage
diff --git a/Engine/Source/Developer/AITestSuite/Classes/MockAI.h b/Engine/Source/Developer/AITestSuite/Classes/MockAI.h
--- a/Engine/Source/Developer/AITestSuite/Classes/MockAI.h
+++ b/Engine/Source/Developer/AITestSuite/Classes/MockAI.h
@@ -60,7 +60,4 @@ class UMockAI : public UObject
 	void SetEnableTicking(bool bShouldTick);
 
 	virtual void TickMe(float DeltaTime);
-
-	// UObject implementation
-	virtual bool IsDestructionThreadSafe() const override { return false; }
 };
diff --git a/Engine/Source/Runtime/AIModule/Classes/AISubsystem.h b/Engine/Source/Runtime/AIModule/Classes/AISubsystem.h
--- a/Engine/Source/Runtime/AIModule/Classes/AISubsystem.h
+++ b/Engine/Source/Runtime/AIModule/Classes/AISubsystem.h
@@ -32,9 +32,5 @@ public:
 	virtual TStatId GetStatId() const override;
 	// FTickableGameObject end
 
-	// UObject begin
-	virtual bool IsDestructionThreadSafe() const override { return false; }
-	// UObject end
-
 	UWorld* GetWorldFast() const { return AISystem ? AISystem->GetOuterWorld() : GetOuter()->GetWorld(); }
 };
diff --git a/Engine/Source/Runtime/AugmentedReality/Public/ARBlueprintProxy.h b/Engine/Source/Runtime/AugmentedReality/Public/ARBlueprintProxy.h
--- a/Engine/Source/Runtime/AugmentedReality/Public/ARBlueprintProxy.h
+++ b/Engine/Source/Runtime/AugmentedReality/Public/ARBlueprintProxy.h
@@ -27,10 +27,6 @@ public:
 	virtual bool IsTickable() const override { return !HasAnyFlags(RF_ClassDefaultObject) && bShouldTick; }
 	virtual TStatId GetStatId() const override { RETURN_QUICK_DECLARE_CYCLE_STAT(UARBaseAsyncTaskBlueprintProxy, STATGROUP_Tickables); }
 	//~ End FTickableObject Interface
-	
-	//~ Begin UObject Interface
-	virtual bool IsDestructionThreadSafe() const override { return false; }
-	//~ End UObject Interface
 
 	virtual void ReportSuccess() { check(0); }
 	virtual void ReportFailure() { check(0); }
