Apple mobile truncation warning fixes

#rnx

diff --git a/Engine/Plugins/Runtime/AppleImageUtils/Source/AppleImageUtils/Private/AppleImageUtilsPlugin.cpp b/Engine/Plugins/Runtime/AppleImageUtils/Source/AppleImageUtils/Private/AppleImageUtilsPlugin.cpp
--- a/Engine/Plugins/Runtime/AppleImageUtils/Source/AppleImageUtils/Private/AppleImageUtilsPlugin.cpp
+++ b/Engine/Plugins/Runtime/AppleImageUtils/Source/AppleImageUtils/Private/AppleImageUtilsPlugin.cpp
@@ -449,8 +449,8 @@ void FAppleImageUtilsPlugin::ConvertToJPEG(CIImage* SourceImage, TArray<uint8>&
 	
 	if (ConvertedImageSize)
 	{
-		ConvertedImageSize->X = Image.extent.size.width;
-		ConvertedImageSize->Y = Image.extent.size.height;
+		ConvertedImageSize->X = (int32)Image.extent.size.width;
+		ConvertedImageSize->Y = (int32)Image.extent.size.height;
 	}
 
 	// This will perform the work on the GPU or inline on this thread
diff --git a/Engine/Plugins/Runtime/AppleImageUtils/Source/AppleImageUtils/Public/IAppleImageUtilsPlugin.h b/Engine/Plugins/Runtime/AppleImageUtils/Source/AppleImageUtils/Public/IAppleImageUtilsPlugin.h
--- a/Engine/Plugins/Runtime/AppleImageUtils/Source/AppleImageUtils/Public/IAppleImageUtilsPlugin.h
+++ b/Engine/Plugins/Runtime/AppleImageUtils/Source/AppleImageUtils/Public/IAppleImageUtilsPlugin.h
@@ -50,7 +50,7 @@ public:
 	virtual bool IsDone() const override { return bIsDone; }
 	virtual bool HadError() const override { return bHadError; }
 	virtual FString GetErrorReason() const override { return Error; }
-	virtual float GetElapsedTime() const override { return (EndTime - StartTime) * 1000.f; }
+	virtual float GetElapsedTime() const override { return (float)(EndTime - StartTime) * 1000.f; }
 	virtual FIntPoint GetConvertedImageSize() const override { return FIntPoint::ZeroValue; };
 
 protected:
