A couple more truncation fixes for iOS

#rnx

diff --git a/Engine/Plugins/Runtime/AppleImageUtils/Source/AppleImageUtils/Private/AppleImageUtilsPlugin.cpp b/Engine/Plugins/Runtime/AppleImageUtils/Source/AppleImageUtils/Private/AppleImageUtilsPlugin.cpp
--- a/Engine/Plugins/Runtime/AppleImageUtils/Source/AppleImageUtils/Private/AppleImageUtilsPlugin.cpp
+++ b/Engine/Plugins/Runtime/AppleImageUtils/Source/AppleImageUtils/Private/AppleImageUtilsPlugin.cpp
@@ -477,8 +477,8 @@ void FAppleImageUtilsPlugin::ConvertToHEIF(CIImage* SourceImage, TArray<uint8>&
 	
 	if (ConvertedImageSize)
 	{
-		ConvertedImageSize->X = Image.extent.size.width;
-		ConvertedImageSize->Y = Image.extent.size.height;
+		ConvertedImageSize->X = (int32)Image.extent.size.width;
+		ConvertedImageSize->Y = (int32)Image.extent.size.height;
 	}
 
 	// This will perform the work on the GPU or inline on this thread
@@ -504,8 +504,8 @@ void FAppleImageUtilsPlugin::ConvertToPNG(CIImage* SourceImage, TArray<uint8>& O
 	
 	if (ConvertedImageSize)
 	{
-		ConvertedImageSize->X = Image.extent.size.width;
-		ConvertedImageSize->Y = Image.extent.size.height;
+		ConvertedImageSize->X = (int32)Image.extent.size.width;
+		ConvertedImageSize->Y = (int32)Image.extent.size.height;
 	}
 
 	// This will perform the work on the GPU or inline on this thread
@@ -531,8 +531,8 @@ void FAppleImageUtilsPlugin::ConvertToTIFF(CIImage* SourceImage, TArray<uint8>&
 	
 	if (ConvertedImageSize)
 	{
-		ConvertedImageSize->X = Image.extent.size.width;
-		ConvertedImageSize->Y = Image.extent.size.height;
+		ConvertedImageSize->X = (int32)Image.extent.size.width;
+		ConvertedImageSize->Y = (int32)Image.extent.size.height;
 	}
 
 	// This will perform the work on the GPU or inline on this thread
